Move Online session FNames to an OnlineBase module to be shared between OnlineSubsystem and OnlineServices

diff --git a/Engine/Plugins/Online/Experimental/OnlineServices/OnlineServices.uplugin b/Engine/Plugins/Online/Experimental/OnlineServices/OnlineServices.uplugin
--- a/Engine/Plugins/Online/Experimental/OnlineServices/OnlineServices.uplugin
+++ b/Engine/Plugins/Online/Experimental/OnlineServices/OnlineServices.uplugin
@@ -21,5 +21,11 @@
 			"Type": "Runtime",
 			"LoadingPhase": "PostConfigInit"
 		}
+	],
+	"Plugins": [
+		{
+			"Name": "OnlineBase",
+			"Enabled": true
+		}
 	]
 }
\ No newline at end of file
diff --git a/Engine/Plugins/Online/Experimental/OnlineServices/Source/OnlineServicesInterface/OnlineServicesInterface.Build.cs b/Engine/Plugins/Online/Experimental/OnlineServices/Source/OnlineServicesInterface/OnlineServicesInterface.Build.cs
--- a/Engine/Plugins/Online/Experimental/OnlineServices/Source/OnlineServicesInterface/OnlineServicesInterface.Build.cs
+++ b/Engine/Plugins/Online/Experimental/OnlineServices/Source/OnlineServicesInterface/OnlineServicesInterface.Build.cs
@@ -12,7 +12,8 @@ public class OnlineServicesInterface : ModuleRules
 
 		PublicDependencyModuleNames.AddRange(
 			new string[] {
-				"CoreOnline"
+				"CoreOnline",
+				"OnlineBase"
 			}
 		);
 
diff --git a/Engine/Plugins/Online/OnlineBase/OnlineBase.uplugin b/Engine/Plugins/Online/OnlineBase/OnlineBase.uplugin
new file mode 100644
--- /dev/null
+++ b/Engine/Plugins/Online/OnlineBase/OnlineBase.uplugin
@@ -0,0 +1,19 @@
+{
+	"FileVersion": 3,
+
+	"FriendlyName": "Online Base",
+	"Version": 1,
+	"VersionName": "1.0",
+	"Description": "Shared code online subsystem (OSSv1) and online service (OSSv2) interfaces.",
+	"Category": "Online Platform",
+	"CreatedBy": "Epic Games, Inc.",
+	"CreatedByURL": "http://epicgames.com",
+	"EnabledByDefault": true,
+
+	"Modules": [
+		{
+			"Name": "OnlineBase",
+			"Type": "Runtime"
+		}
+	]
+}
\ No newline at end of file
diff --git a/Engine/Plugins/Online/OnlineBase/Source/OnlineBase.Build.cs b/Engine/Plugins/Online/OnlineBase/Source/OnlineBase.Build.cs
new file mode 100644
--- /dev/null
+++ b/Engine/Plugins/Online/OnlineBase/Source/OnlineBase.Build.cs
@@ -0,0 +1,19 @@
+// Copyright Epic Games, Inc. All Rights Reserved.
+
+using UnrealBuildTool;
+
+namespace UnrealBuildTool.Rules
+{
+	public class OnlineBase : ModuleRules
+	{
+		public OnlineBase(ReadOnlyTargetRules Target) : base(Target)
+		{
+			// NOTE:  OnlineBase cannot depend on Engine!
+			PrivateDependencyModuleNames.AddRange(
+				new string[] {
+					"Core"
+				}
+			);
+		}
+	}
+}
\ No newline at end of file
diff --git a/Engine/Plugins/Online/OnlineBase/Source/Private/OnlineBaseModule.cpp b/Engine/Plugins/Online/OnlineBase/Source/Private/OnlineBaseModule.cpp
new file mode 100644
--- /dev/null
+++ b/Engine/Plugins/Online/OnlineBase/Source/Private/OnlineBaseModule.cpp
@@ -0,0 +1,23 @@
+// Copyright Epic Games, Inc. All Rights Reserved.
+
+#include "Modules/ModuleManager.h"
+
+/**
+ * Online base module class - to support this otherwise header-only module to be loaded
+ */
+class FOnlineBaseModule : public IModuleInterface
+{
+public:
+	// IModuleInterface
+	virtual bool SupportsDynamicReloading() override
+	{
+		return false;
+	}
+
+	virtual bool SupportsAutomaticShutdown() override
+	{
+		return false;
+	}
+};
+
+IMPLEMENT_MODULE(FOnlineBaseModule, OnlineBase);
diff --git a/Engine/Plugins/Online/OnlineBase/Source/Public/Online/OnlineSessionNames.h b/Engine/Plugins/Online/OnlineBase/Source/Public/Online/OnlineSessionNames.h
new file mode 100644
--- /dev/null
+++ b/Engine/Plugins/Online/OnlineBase/Source/Public/Online/OnlineSessionNames.h
@@ -0,0 +1,150 @@
+// Copyright Epic Games, Inc. All Rights Reserved.
+
+#pragma once
+
+#include "CoreMinimal.h"
+
+/** default beacon port, if not specified by other means */
+#define DEFAULT_BEACON_PORT 15000
+
+/**
+ * Session settings
+ */
+/** Setting describing the name of the current map (value is FString) */
+#define SETTING_MAPNAME FName(TEXT("MAPNAME"))
+/** Setting describing the number of bots in the session (value is int32) */
+#define SETTING_NUMBOTS FName(TEXT("NUMBOTS"))
+/** Setting describing the game mode of the session (value is FString) */
+#define SETTING_GAMEMODE FName(TEXT("GAMEMODE"))
+/** Setting describing the beacon host port (value is int32) */
+#define SETTING_BEACONPORT FName(TEXT("BEACONPORT"))
+/** Server responds to Qos beacon requests (value is int32) */
+#define SETTING_QOS FName(TEXT("QOS"))
+/** Setting describing the region of the world you are in (value is FString) */
+#define SETTING_REGION FName(TEXT("REGION"))
+/** Setting describing the a specific subregion preference within a region (value is FString) */
+#define SETTING_SUBREGION FName(TEXT("SUBREGION"))
+/** Setting describing the unique id of a datacenter (value is FString) */
+#define SETTING_DCID FName(TEXT("DCID"))
+/** Number of players needed to fill out this session (value is int32) */
+#define SETTING_NEEDS FName(TEXT("NEEDS"))
+/** Second key for "needs" because can't set same value with two criteria (value is int32) */
+#define SETTING_NEEDSSORT FName(TEXT("NEEDSSORT"))
+/** Setting describing the session key (value is FString) */
+#define SETTING_SESSIONKEY FName(TEXT("SESSIONKEY"))
+/** Setting describing the session key (value is FString) */
+#define SETTING_ALLOWBROADCASTING FName(TEXT("ALLOWBROADCASTING"))
+/** Setting defining broadcaster (value is bool) */
+#define SETTING_BROADCASTER FName(TEXT("BROADCASTER"))
+/** Setting defining if we are a party member spectator (value is bool) */
+#define SETTING_PARTYMEMBERSPECTATOR FName(TEXT("PARTYMEMBERSPECTATOR"))
+/** Setting describing the number of spectator slots in the session (value is int32) */
+#define SETTING_MAXSPECTATORS FName(TEXT("MAXSPECTATORS"))
+
+/** 8 user defined integer params to be used when filtering searches for sessions */
+#define SETTING_CUSTOMSEARCHINT1 FName(TEXT("CUSTOMSEARCHINT1"))
+#define SETTING_CUSTOMSEARCHINT2 FName(TEXT("CUSTOMSEARCHINT2"))
+#define SETTING_CUSTOMSEARCHINT3 FName(TEXT("CUSTOMSEARCHINT3"))
+#define SETTING_CUSTOMSEARCHINT4 FName(TEXT("CUSTOMSEARCHINT4"))
+#define SETTING_CUSTOMSEARCHINT5 FName(TEXT("CUSTOMSEARCHINT5"))
+#define SETTING_CUSTOMSEARCHINT6 FName(TEXT("CUSTOMSEARCHINT6"))
+#define SETTING_CUSTOMSEARCHINT7 FName(TEXT("CUSTOMSEARCHINT7"))
+#define SETTING_CUSTOMSEARCHINT8 FName(TEXT("CUSTOMSEARCHINT8"))
+
+/**
+ * Setting describing whether host migration is enabled for the session.  This typically means when the host leaves, 
+ * the online service picks a new host instead of destroying the session.  Not supported by all online services.
+ */
+#define SETTING_HOST_MIGRATION FName(TEXT("HOSTMIGRATION"))
+
+/** TODO ONLINE Settings to consider */
+/** The server's nonce for this session */
+/** Whether the game is an invitation or searched for game */
+/** The ping of the server in milliseconds (-1 means the server was unreachable) */
+/** Whether this server is a dedicated server or not */
+/** Represents how good a match this is in a range from 0 to 1 */
+/** Whether there is a skill update in progress or not (don't do multiple at once) */
+/** Whether to shrink the session slots when a player leaves the match or not */
+
+// These are Xbox One specific setting to be used in FOnlineSessionSettings as keys; technically this can be put in OnlineSubsystemSettings.h
+
+/** Custom settings to be associated with Session (value is FString)*/
+#define SETTING_CUSTOM FName(TEXT("CUSTOM"))
+/** Session is marked as using party enabled session flag (value is bool)*/
+#define SETTING_PARTY_ENABLED_SESSION FName(TEXT("PARTYENABLEDSESSION"))
+/** Hopper name to search within (value is FString) */
+#define SETTING_MATCHING_HOPPER FName(TEXT("MATCHHOPPER"))
+/** Timeout for search (value is float seconds) */
+#define SETTING_MATCHING_TIMEOUT FName(TEXT("MATCHTIMEOUT"))
+/** Match attributes (value is FString) */
+#define SETTING_MATCHING_ATTRIBUTES	FName(TEXT("MATCHATTRIBUTES"))
+/** Match attributes (value is FString) */
+#define SETTING_MATCH_MEMBERS_JSON FName(TEXT("MATCHMEMBERS"))
+/** Session member constant custom json (value is FString) - This is to be used with FString::Printf, populated with user xuid */
+#define SETTING_SESSION_MEMBER_CONSTANT_CUSTOM_JSON_XUID_PREFIX TEXT("SESSIONMEMBERCONSTANTCUSTOMJSON")
+/** Set self as host (value is int32)*/
+#define SETTING_MAX_RESULT FName(TEXT("MAXRESULT"))
+/** Set self as host (value is int32)*/
+#define SETTING_CONTRACT_VERSION_FILTER FName(TEXT("CONTRACTVERSIONFILTER"))
+/** Set self as host (value is int32)*/
+#define SETTING_FIND_PRIVATE_SESSIONS FName(TEXT("FINDPRIVATESESSIONS"))
+/** Set self as host (value is int32)*/
+#define SETTING_FIND_RESERVED_SESSIONS FName(TEXT("FINDRESERVEDSESSIONS"))
+/** Set self as host (value is int32)*/
+#define SETTING_FIND_INACTIVE_SESSIONS FName(TEXT("FINDINACTIVESESSIONS"))
+/** Set self as host (value is int32)*/
+#define SETTING_MULTIPLAYER_VISIBILITY FName(TEXT("MULTIPLAYERVISIBILITY"))
+/** Session template name, must match a session template name from the service config. */
+#define SETTING_SESSION_TEMPLATE_NAME FName(TEXT("SESSIONTEMPLATENAME"))
+/** Session change number from shoulder tapped version */
+#define SETTING_CHANGE_NUMBER FName(TEXT("CHANGENUMBER"))
+/** Session can attempt arbiter migration */
+#define SETTING_ALLOW_ARBITER_MIGRATION FName(TEXT("ALLOWARBITERMIGRATION"))
+/** Session preservation type for matchmaking */
+#define SETTING_MATCHING_PRESERVE_SESSION FName(TEXT("PRESERVESESSIONALWAYS"))
+/** Matchmade game session URI for join in progress/invites */
+#define  SETTING_GAME_SESSION_URI FName(TEXT("GAMESESSIONURI"))
+/** Session member group identifier (value is FString) - This is to be used with FString::Printf, populated with user xuid - Field required with Team Based matchmaking*/
+#define SETTING_GROUP_NAME_PREFIX TEXT("USERGROUPNAME")
+/** Information to join a third-party auxiliary session (value is FString) */
+#define SETTING_CUSTOM_JOIN_INFO FName(TEXT("CUSTOMJOININFO"))
+/** What verbosity of session updates to subscribe to (value is FString, comma-separated if multiple values needed)*/
+#define SETTING_SESSION_SUBSCRIPTION_TYPES FName(TEXT("SESSIONSUBSCRIPTIONTYPE"))
+
+/**
+ * Search settings
+ */
+
+/** Search only for dedicated servers (value is true/false) */
+#define SEARCH_DEDICATED_ONLY FName(TEXT("DEDICATEDONLY"))
+/** Search for empty servers only (value is true/false) */
+#define SEARCH_EMPTY_SERVERS_ONLY FName(TEXT("EMPTYONLY"))
+/** Search for non empty servers only (value is true/false) */
+#define SEARCH_NONEMPTY_SERVERS_ONLY FName(TEXT("NONEMPTYONLY"))
+/** Search for secure servers only (value is true/false) */
+#define SEARCH_SECURE_SERVERS_ONLY FName(TEXT("SECUREONLY"))
+/** Search for presence sessions only (value is true/false) */
+#define SEARCH_PRESENCE FName(TEXT("PRESENCESEARCH"))
+/** Search for a match with min player availability (value is int) */
+#define SEARCH_MINSLOTSAVAILABLE FName(TEXT("MINSLOTSAVAILABLE"))
+/** Exclude all matches where any unique ids in a given array are present (value is string of the form "uniqueid1;uniqueid2;uniqueid3") */
+#define SEARCH_EXCLUDE_UNIQUEIDS FName(TEXT("EXCLUDEUNIQUEIDS"))
+/** User ID to search for session of */
+#define SEARCH_USER FName(TEXT("SEARCHUSER"))
+/** Keywords to match in session search */
+#define SEARCH_KEYWORDS FName(TEXT("SEARCHKEYWORDS"))
+/** The matchmaking queue name to matchmake in, e.g. "TeamDeathmatch" (value is string) */
+#define SEARCH_MATCHMAKING_QUEUE FName(TEXT("MATCHMAKINGQUEUE"))
+/** If set, use the named Xbox Live hopper to find a session via matchmaking (value is a string) */
+#define SEARCH_XBOX_LIVE_HOPPER_NAME FName(TEXT("LIVEHOPPERNAME"))
+/** Which session template from the service configuration to use */
+#define SEARCH_XBOX_LIVE_SESSION_TEMPLATE_NAME FName(TEXT("LIVESESSIONTEMPLATE"))
+/** Selection method used to determine which match to join when multiple are returned (valid only on Switch) */
+#define SEARCH_SWITCH_SELECTION_METHOD FName(TEXT("SWITCHSELECTIONMETHOD"))
+/** Whether to use lobbies vs sessions */
+#define SEARCH_LOBBIES FName(TEXT("LOBBYSEARCH"))
+
+// User attributes for searching (FSessionMatchmakingUser::Attributes)
+/** Team a user is searching for */
+#define SEARCH_USER_ATTRIBUTE_TEAM TEXT("TEAM")
+
diff --git a/Engine/Plugins/Online/OnlineSubsystem/OnlineSubsystem.uplugin b/Engine/Plugins/Online/OnlineSubsystem/OnlineSubsystem.uplugin
--- a/Engine/Plugins/Online/OnlineSubsystem/OnlineSubsystem.uplugin
+++ b/Engine/Plugins/Online/OnlineSubsystem/OnlineSubsystem.uplugin
@@ -22,5 +22,11 @@
 			"Name": "OnlineSubsystem",
 			"LoadingPolicy": "Always"
 		}
+	],
+	"Plugins": [
+		{
+			"Name": "OnlineBase",
+			"Enabled": true
+		}
 	]
 }
\ No newline at end of file
diff --git a/Engine/Plugins/Online/OnlineSubsystem/Source/OnlineSubsystem.Build.cs b/Engine/Plugins/Online/OnlineSubsystem/Source/OnlineSubsystem.Build.cs
--- a/Engine/Plugins/Online/OnlineSubsystem/Source/OnlineSubsystem.Build.cs
+++ b/Engine/Plugins/Online/OnlineSubsystem/Source/OnlineSubsystem.Build.cs
@@ -12,6 +12,7 @@ public class OnlineSubsystem : ModuleRules
 			new string[] {
 				"Json",
 				"CoreOnline",
+				"OnlineBase",
 				"SignalProcessing"
 			}
 		);
diff --git a/Engine/Plugins/Online/OnlineSubsystem/Source/Public/OnlineSessionSettings.h b/Engine/Plugins/Online/OnlineSubsystem/Source/Public/OnlineSessionSettings.h
--- a/Engine/Plugins/Online/OnlineSubsystem/Source/Public/OnlineSessionSettings.h
+++ b/Engine/Plugins/Online/OnlineSubsystem/Source/Public/OnlineSessionSettings.h
@@ -8,58 +8,8 @@
 #include "OnlineKeyValuePair.h"
 #include "OnlineSubsystemPackage.h"
 
-/** default beacon port, if not specified by other means */
-#define DEFAULT_BEACON_PORT 15000
-
-/** Setting describing the name of the current map (value is FString) */
-#define SETTING_MAPNAME FName(TEXT("MAPNAME"))
-/** Setting describing the number of bots in the session (value is int32) */
-#define SETTING_NUMBOTS FName(TEXT("NUMBOTS"))
-/** Setting describing the game mode of the session (value is FString) */
-#define SETTING_GAMEMODE FName(TEXT("GAMEMODE"))
-/** Setting describing the beacon host port (value is int32) */
-#define SETTING_BEACONPORT FName(TEXT("BEACONPORT"))
-/** Server responds to Qos beacon requests (value is int32) */
-#define SETTING_QOS FName(TEXT("QOS"))
-/** Setting describing the region of the world you are in (value is FString) */
-#define SETTING_REGION FName(TEXT("REGION"))
-/** Setting describing the a specific subregion preference within a region (value is FString) */
-#define SETTING_SUBREGION FName(TEXT("SUBREGION"))
-/** Setting describing the unique id of a datacenter (value is FString) */
-#define SETTING_DCID FName(TEXT("DCID"))
-/** Number of players needed to fill out this session (value is int32) */
-#define SETTING_NEEDS FName(TEXT("NEEDS"))
-/** Second key for "needs" because can't set same value with two criteria (value is int32) */
-#define SETTING_NEEDSSORT FName(TEXT("NEEDSSORT"))
-/** Setting describing the session key (value is FString) */
-#define SETTING_SESSIONKEY FName(TEXT("SESSIONKEY"))
-/** Setting describing the session key (value is FString) */
-#define SETTING_ALLOWBROADCASTING FName(TEXT("ALLOWBROADCASTING"))
-/** Setting defining broadcaster (value is bool) */
-#define SETTING_BROADCASTER FName(TEXT("BROADCASTER"))
-/** Setting defining if we are a party member spectator (value is bool) */
-#define SETTING_PARTYMEMBERSPECTATOR FName(TEXT("PARTYMEMBERSPECTATOR"))
-/** Setting describing the number of spectator slots in the session (value is int32) */
-#define SETTING_MAXSPECTATORS FName(TEXT("MAXSPECTATORS"))
-
-/** 8 user defined integer params to be used when filtering searches for sessions */
-#define SETTING_CUSTOMSEARCHINT1 FName(TEXT("CUSTOMSEARCHINT1"))
-#define SETTING_CUSTOMSEARCHINT2 FName(TEXT("CUSTOMSEARCHINT2"))
-#define SETTING_CUSTOMSEARCHINT3 FName(TEXT("CUSTOMSEARCHINT3"))
-#define SETTING_CUSTOMSEARCHINT4 FName(TEXT("CUSTOMSEARCHINT4"))
-#define SETTING_CUSTOMSEARCHINT5 FName(TEXT("CUSTOMSEARCHINT5"))
-#define SETTING_CUSTOMSEARCHINT6 FName(TEXT("CUSTOMSEARCHINT6"))
-#define SETTING_CUSTOMSEARCHINT7 FName(TEXT("CUSTOMSEARCHINT7"))
-#define SETTING_CUSTOMSEARCHINT8 FName(TEXT("CUSTOMSEARCHINT8"))
-
-/** TODO ONLINE Settings to consider */
-/** The server's nonce for this session */
-/** Whether the game is an invitation or searched for game */
-/** The ping of the server in milliseconds (-1 means the server was unreachable) */
-/** Whether this server is a dedicated server or not */
-/** Represents how good a match this is in a range from 0 to 1 */
-/** Whether there is a skill update in progress or not (don't do multiple at once) */
-/** Whether to shrink the session slots when a player leaves the match or not */
+// NOTE:  Session setting names have been moved to OnlineBase in OnlineSessionNames.h
+#include "Online/OnlineSessionNames.h"
 
 #define INVALID_SESSION_SETTING_ID -1
 
@@ -614,39 +564,6 @@ public:
 	}
 };
 
-/** Search only for dedicated servers (value is true/false) */
-#define SEARCH_DEDICATED_ONLY FName(TEXT("DEDICATEDONLY"))
-/** Search for empty servers only (value is true/false) */
-#define SEARCH_EMPTY_SERVERS_ONLY FName(TEXT("EMPTYONLY"))
-/** Search for non empty servers only (value is true/false) */
-#define SEARCH_NONEMPTY_SERVERS_ONLY FName(TEXT("NONEMPTYONLY"))
-/** Search for secure servers only (value is true/false) */
-#define SEARCH_SECURE_SERVERS_ONLY FName(TEXT("SECUREONLY"))
-/** Search for presence sessions only (value is true/false) */
-#define SEARCH_PRESENCE FName(TEXT("PRESENCESEARCH"))
-/** Search for a match with min player availability (value is int) */
-#define SEARCH_MINSLOTSAVAILABLE FName(TEXT("MINSLOTSAVAILABLE"))
-/** Exclude all matches where any unique ids in a given array are present (value is string of the form "uniqueid1;uniqueid2;uniqueid3") */
-#define SEARCH_EXCLUDE_UNIQUEIDS FName(TEXT("EXCLUDEUNIQUEIDS"))
-/** User ID to search for session of */
-#define SEARCH_USER FName(TEXT("SEARCHUSER"))
-/** Keywords to match in session search */
-#define SEARCH_KEYWORDS FName(TEXT("SEARCHKEYWORDS"))
-/** The matchmaking queue name to matchmake in, e.g. "TeamDeathmatch" (value is string) */
-#define SEARCH_MATCHMAKING_QUEUE FName(TEXT("MATCHMAKINGQUEUE"))
-/** If set, use the named Xbox Live hopper to find a session via matchmaking (value is a string) */
-#define SEARCH_XBOX_LIVE_HOPPER_NAME FName(TEXT("LIVEHOPPERNAME"))
-/** Which session template from the service configuration to use */
-#define SEARCH_XBOX_LIVE_SESSION_TEMPLATE_NAME FName(TEXT("LIVESESSIONTEMPLATE"))
-/** Selection method used to determine which match to join when multiple are returned (valid only on Switch) */
-#define SEARCH_SWITCH_SELECTION_METHOD FName(TEXT("SWITCHSELECTIONMETHOD"))
-/** Whether to use lobbies vs sessions */
-#define SEARCH_LOBBIES FName(TEXT("LOBBYSEARCH"))
-
-// User attributes for searching (FSessionMatchmakingUser::Attributes)
-/** Team a user is searching for */
-#define SEARCH_USER_ATTRIBUTE_TEAM TEXT("TEAM")
-
 /**
  * Encapsulation of a search for sessions request.
  * Contains all the search parameters and any search results returned after
diff --git a/Engine/Plugins/Online/OnlineSubsystem/Source/Public/OnlineSubsystemSessionSettings.h b/Engine/Plugins/Online/OnlineSubsystem/Source/Public/OnlineSubsystemSessionSettings.h
--- a/Engine/Plugins/Online/OnlineSubsystem/Source/Public/OnlineSubsystemSessionSettings.h
+++ b/Engine/Plugins/Online/OnlineSubsystem/Source/Public/OnlineSubsystemSessionSettings.h
@@ -2,56 +2,5 @@
 
 #pragma once
 
-#include "CoreMinimal.h"
-#include "UObject/NameTypes.h"
-#include "OnlineSubsystemPackage.h"
-
-// These are Xbox One specific setting to be used in FOnlineSessionSettings as keys; technically this can be put in OnlineSubsystemSettings.h
-
-/** Custom settings to be associated with Session (value is FString)*/
-#define SETTING_CUSTOM FName(TEXT("CUSTOM"))
-/** Session is marked as using party enabled session flag (value is bool)*/
-#define SETTING_PARTY_ENABLED_SESSION FName(TEXT("PARTYENABLEDSESSION"))
-/** Hopper name to search within (value is FString) */
-#define SETTING_MATCHING_HOPPER FName(TEXT("MATCHHOPPER"))
-/** Timeout for search (value is float seconds) */
-#define SETTING_MATCHING_TIMEOUT FName(TEXT("MATCHTIMEOUT"))
-/** Match attributes (value is FString) */
-#define SETTING_MATCHING_ATTRIBUTES	FName(TEXT("MATCHATTRIBUTES"))
-/** Match attributes (value is FString) */
-#define SETTING_MATCH_MEMBERS_JSON FName(TEXT("MATCHMEMBERS"))
-/** Session member constant custom json (value is FString) - This is to be used with FString::Printf, populated with user xuid */
-#define SETTING_SESSION_MEMBER_CONSTANT_CUSTOM_JSON_XUID_PREFIX TEXT("SESSIONMEMBERCONSTANTCUSTOMJSON")
-/** Set self as host (value is int32)*/
-#define SETTING_MAX_RESULT FName(TEXT("MAXRESULT"))
-/** Set self as host (value is int32)*/
-#define SETTING_CONTRACT_VERSION_FILTER FName(TEXT("CONTRACTVERSIONFILTER"))
-/** Set self as host (value is int32)*/
-#define SETTING_FIND_PRIVATE_SESSIONS FName(TEXT("FINDPRIVATESESSIONS"))
-/** Set self as host (value is int32)*/
-#define SETTING_FIND_RESERVED_SESSIONS FName(TEXT("FINDRESERVEDSESSIONS"))
-/** Set self as host (value is int32)*/
-#define SETTING_FIND_INACTIVE_SESSIONS FName(TEXT("FINDINACTIVESESSIONS"))
-/** Set self as host (value is int32)*/
-#define SETTING_MULTIPLAYER_VISIBILITY FName(TEXT("MULTIPLAYERVISIBILITY"))
-/** Session template name, must match a session template name from the service config. */
-#define SETTING_SESSION_TEMPLATE_NAME FName(TEXT("SESSIONTEMPLATENAME"))
-/** Session change number from shoulder tapped version */
-#define SETTING_CHANGE_NUMBER FName(TEXT("CHANGENUMBER"))
-/** Session can attempt arbiter migration */
-#define SETTING_ALLOW_ARBITER_MIGRATION FName(TEXT("ALLOWARBITERMIGRATION"))
-/** Session preservation type for matchmaking */
-#define SETTING_MATCHING_PRESERVE_SESSION FName(TEXT("PRESERVESESSIONALWAYS"))
-/** Matchmade game session URI for join in progress/invites */
-#define  SETTING_GAME_SESSION_URI FName(TEXT("GAMESESSIONURI"))
-/** Session member group identifier (value is FString) - This is to be used with FString::Printf, populated with user xuid - Field required with Team Based matchmaking*/
-#define SETTING_GROUP_NAME_PREFIX TEXT("USERGROUPNAME")
-/** Information to join a third-party auxiliary session (value is FString) */
-#define SETTING_CUSTOM_JOIN_INFO FName(TEXT("CUSTOMJOININFO"))
-/** What verbosity of session updates to subscribe to (value is FString, comma-separated if multiple values needed)*/
-#define SETTING_SESSION_SUBSCRIPTION_TYPES FName(TEXT("SESSIONSUBSCRIPTIONTYPE"))
-
-// These are PS4 specific settings to be used in FOnlineSessionSettings as keys
-
-/** Enables host migration for PS4 sessions, which is handled on Sony's servers. This means that all clients can update the session, we do not get callbacks when we become the host or a chance to set the host, it is picked automatically. */
-#define SETTING_HOST_MIGRATION FName(TEXT("HOSTMIGRATION"))
+// NOTE:  Session setting names have been moved to OnlineBase in OnlineSessionNames.h
+#include "Online/OnlineSessionNames.h"
