Movement of LANBeacon and related files from OnlineSubsystem to OnlineBase, and related refactoring.

diff --git a/Engine/Plugins/Online/Experimental/OnlineServicesNull/Source/OnlineServicesNull.Build.cs b/Engine/Plugins/Online/Experimental/OnlineServicesNull/Source/OnlineServicesNull.Build.cs
--- a/Engine/Plugins/Online/Experimental/OnlineServicesNull/Source/OnlineServicesNull.Build.cs
+++ b/Engine/Plugins/Online/Experimental/OnlineServicesNull/Source/OnlineServicesNull.Build.cs
@@ -21,6 +21,7 @@ public class OnlineServicesNull : ModuleRules
 		PrivateDependencyModuleNames.AddRange(
 			new string[] { 
 				"ApplicationCore",
+				"OnlineBase"
 			}
 		);
 
diff --git a/Engine/Plugins/Online/Experimental/OnlineServicesNull/Source/Public/Online/LobbiesNull.h b/Engine/Plugins/Online/Experimental/OnlineServicesNull/Source/Public/Online/LobbiesNull.h
--- a/Engine/Plugins/Online/Experimental/OnlineServicesNull/Source/Public/Online/LobbiesNull.h
+++ b/Engine/Plugins/Online/Experimental/OnlineServicesNull/Source/Public/Online/LobbiesNull.h
@@ -4,7 +4,7 @@
 #include "CoreMinimal.h"
 #include "Online/LobbiesCommon.h"
 #include "Online/Lobbies.h"
-#include "LANBeacon.h"
+#include "Online/LANBeacon.h"
 
 namespace UE::Online {
 
diff --git a/Engine/Plugins/Online/Experimental/OnlineServicesNull/Source/Public/Online/NboSerializerNullSvc.h b/Engine/Plugins/Online/Experimental/OnlineServicesNull/Source/Public/Online/NboSerializerNullSvc.h
--- a/Engine/Plugins/Online/Experimental/OnlineServicesNull/Source/Public/Online/NboSerializerNullSvc.h
+++ b/Engine/Plugins/Online/Experimental/OnlineServicesNull/Source/Public/Online/NboSerializerNullSvc.h
@@ -2,8 +2,8 @@
 
 #pragma once
 
-#include "CoreMinimal.h"
-#include "NboSerializer.h"
+#include "Online/NboSerializer.h"
+
 #include "Online/AuthNull.h"
 #include "Online/LobbiesNull.h"
 #include "Online/SessionsNull.h"
@@ -12,7 +12,7 @@
 /**
  * Serializes data in network byte order form into a buffer
  */
- namespace UE::Online {
+namespace UE::Online {
 
 class FNboSerializeToBufferNullSvc : public FNboSerializeToBuffer
 {
diff --git a/Engine/Plugins/Online/Experimental/OnlineServicesNull/Source/Public/Online/SessionsNull.h b/Engine/Plugins/Online/Experimental/OnlineServicesNull/Source/Public/Online/SessionsNull.h
--- a/Engine/Plugins/Online/Experimental/OnlineServicesNull/Source/Public/Online/SessionsNull.h
+++ b/Engine/Plugins/Online/Experimental/OnlineServicesNull/Source/Public/Online/SessionsNull.h
@@ -2,7 +2,7 @@
 #pragma once
 
 #include "CoreMinimal.h"
-#include "LANBeacon.h"
+#include "Online/LANBeacon.h"
 #include "Online/SessionsCommon.h"
 
 namespace UE::Online {
diff --git a/Engine/Plugins/Online/OnlineBase/Source/OnlineBase.Build.cs b/Engine/Plugins/Online/OnlineBase/Source/OnlineBase.Build.cs
--- a/Engine/Plugins/Online/OnlineBase/Source/OnlineBase.Build.cs
+++ b/Engine/Plugins/Online/OnlineBase/Source/OnlineBase.Build.cs
@@ -8,10 +8,17 @@ namespace UnrealBuildTool.Rules
 	{
 		public OnlineBase(ReadOnlyTargetRules Target) : base(Target)
 		{
+			PublicDependencyModuleNames.AddRange(
+				new string[] {
+					"Sockets"
+				}
+			);
+
 			// NOTE:  OnlineBase cannot depend on Engine!
 			PrivateDependencyModuleNames.AddRange(
 				new string[] {
-					"Core"
+					"Core",
+					"CoreUObject"
 				}
 			);
 		}
diff --git a/Engine/Plugins/Online/OnlineBase/Source/Private/Online/LANBeacon.cpp b/Engine/Plugins/Online/OnlineBase/Source/Private/Online/LANBeacon.cpp
new file mode 100644
--- /dev/null
+++ b/Engine/Plugins/Online/OnlineBase/Source/Private/Online/LANBeacon.cpp
@@ -0,0 +1,469 @@
+// Copyright Epic Games, Inc. All Rights Reserved.
+
+#include "Online/LANBeacon.h"
+#include "Misc/FeedbackContext.h"
+#include "UObject/CoreNet.h"
+#include "SocketSubsystem.h"
+#include "Sockets.h"
+#include "Online/NboSerializer.h"
+
+/** Sets the broadcast address for this object */
+FLanBeacon::FLanBeacon(void) 
+	: ListenSocket(NULL),
+	  SockAddr(ISocketSubsystem::Get(PLATFORM_SOCKETSUBSYSTEM)->CreateInternetAddr())
+{
+}
+
+/** Frees the broadcast socket */
+FLanBeacon::~FLanBeacon(void)
+{
+	ISocketSubsystem* SocketSubsystem = ISocketSubsystem::Get(PLATFORM_SOCKETSUBSYSTEM);
+	SocketSubsystem->DestroySocket(ListenSocket);
+}
+
+/** Return true if there is a valid ListenSocket */
+bool FLanBeacon::IsListenSocketValid() const
+{
+	return (ListenSocket ? true : false);
+}
+
+/**
+ * Initializes the socket
+ *
+ * @param Port the port to listen on
+ *
+ * @return true if both socket was created successfully, false otherwise
+ */
+bool FLanBeacon::Init(int32 Port)
+{
+	ISocketSubsystem* SocketSubsystem = ISocketSubsystem::Get(PLATFORM_SOCKETSUBSYSTEM);
+	bool bSuccess = false;
+	// Now the listen address
+	ListenAddr = SocketSubsystem->GetLocalBindAddr(*GWarn);
+	ListenAddr->SetPort(Port);
+	// Set our broadcast address
+	BroadcastAddr = SocketSubsystem->CreateInternetAddr(ListenAddr->GetProtocolType());
+	BroadcastAddr->SetBroadcastAddress();
+	BroadcastAddr->SetPort(Port);
+	// A temporary "received from" address
+	SockAddr = SocketSubsystem->CreateInternetAddr();
+
+	// Now create and set up our sockets (no VDP)
+	ListenSocket = SocketSubsystem->CreateSocket(NAME_DGram, TEXT("LAN beacon"), ListenAddr->GetProtocolType());
+	if (ListenSocket != NULL)
+	{
+		ListenSocket->SetReuseAddr();
+		ListenSocket->SetNonBlocking();
+		ListenSocket->SetRecvErr();
+		// Bind to our listen port
+		if (ListenSocket->Bind(*ListenAddr))
+		{
+			// Set it to broadcast mode, so we can send on it
+			// NOTE: You must set this to broadcast mode on Xbox 360 or the
+			// secure layer will eat any packets sent
+			bSuccess = ListenSocket->SetBroadcast();
+		}
+		else
+		{
+			UE_LOG(LogTemp, Error, TEXT("Failed to bind listen socket to addr (%s) for LAN beacon"),
+				*ListenAddr->ToString(true));
+		}
+	}
+	else
+	{
+		UE_LOG(LogTemp, Error, TEXT("Failed to create listen socket for LAN beacon"));
+	}
+	return bSuccess && ListenSocket;
+}
+
+/**
+ * Called to poll the socket for pending data. Any data received is placed
+ * in the specified packet buffer
+ *
+ * @param PacketData the buffer to get the socket's packet data
+ * @param BufferSize the size of the packet buffer
+ *
+ * @return the number of bytes read (<= 0 if none or an error)
+ */
+int32 FLanBeacon::ReceivePacket(uint8* PacketData, int32 BufferSize)
+{
+	check(PacketData && BufferSize);
+	// Default to no data being read
+	int32 BytesRead = 0;
+	if (ListenSocket != NULL)
+	{
+		// Read from the socket
+		ListenSocket->RecvFrom(PacketData, BufferSize, BytesRead, *SockAddr);
+		if (BytesRead > 0)
+		{
+			UE_LOG(LogTemp, Verbose, TEXT("Received %d bytes from %s"), BytesRead, *SockAddr->ToString(true));
+		}
+	}
+
+	return BytesRead;
+}
+
+/**
+ * Uses the cached broadcast address to send packet to a subnet
+ *
+ * @param Packet the packet to send
+ * @param Length the size of the packet to send
+ */
+bool FLanBeacon::BroadcastPacket(uint8* Packet, int32 Length)
+{
+	int32 BytesSent = 0;
+	UE_LOG(LogTemp, Verbose, TEXT("Sending %d bytes to %s"), Length, *BroadcastAddr->ToString(true));
+	return ListenSocket->SendTo(Packet, Length, BytesSent, *BroadcastAddr) && (BytesSent == Length);
+}
+
+/**
+* Creates the LAN beacon for queries/advertising servers
+*/
+bool FLANSession::Host(FOnValidQueryPacketDelegate& QueryDelegate)
+{
+	bool bSuccess = false;
+	if (LanBeacon != NULL)
+	{
+		StopLANSession();
+	}
+
+	LanBeaconGuid = FGuid::NewGuid();
+
+	// Bind a socket for LAN beacon activity
+	LanBeacon = new FLanBeacon();
+	if (LanBeacon->Init(LanAnnouncePort))
+	{
+		OnValidQueryPacketDelegates.Add(QueryDelegate);
+		// We successfully created everything so mark the socket as needing polling
+		LanBeaconState = ELanBeaconState::Hosting;
+		bSuccess = true;
+		UE_LOG(LogTemp, Verbose, TEXT("Listening for LAN beacon requests on %d"),	LanAnnouncePort);
+	}
+	else
+	{
+		UE_LOG(LogTemp, Warning, TEXT("Failed to init to LAN beacon %s"),	ISocketSubsystem::Get(PLATFORM_SOCKETSUBSYSTEM)->GetSocketError());
+	}
+	
+	return bSuccess;
+}
+
+/**
+ * Creates the LAN beacon for queries/advertising servers
+ *
+ * @param Nonce unique identifier for this search
+ * @param ResponseDelegate delegate to fire when a server response is received
+ */
+bool FLANSession::Search(FNboSerializeToBuffer& Packet, FOnValidResponsePacketDelegate& ResponseDelegate, FOnSearchingTimeoutDelegate& TimeoutDelegate)
+{
+	bool bSuccess = true;
+	if (LanBeacon != NULL)
+	{
+		StopLANSession();
+	}
+	
+	LanBeaconGuid = FGuid::NewGuid();
+
+	CachedResponseGuids.Empty();
+
+	// Bind a socket for LAN beacon activity
+	LanBeacon = new FLanBeacon();
+	if (LanBeacon->Init(LanAnnouncePort) == false)
+	{
+		UE_LOG(LogTemp, Warning, TEXT("Failed to create socket for lan announce port %s"), ISocketSubsystem::Get(PLATFORM_SOCKETSUBSYSTEM)->GetSocketError());
+		bSuccess = false;
+	}
+
+	// If we have a socket and a nonce, broadcast a discovery packet
+	if (LanBeacon && bSuccess)
+	{
+		// Now kick off our broadcast which hosts will respond to
+		if (LanBeacon->BroadcastPacket(Packet, Packet.GetByteCount()))
+		{
+			UE_LOG(LogTemp, Verbose, TEXT("Sent query packet..."));
+			// We need to poll for the return packets
+			LanBeaconState = ELanBeaconState::Searching;
+			// Set the timestamp for timing out a search
+			LanQueryTimeLeft = LanQueryTimeout;
+
+			RetryData.Reset();
+			RetryData.Append(Packet.GetBuffer().GetData(), Packet.GetByteCount());
+
+			LanQueryRetryTimeLeft = LanQueryRetryTime;
+
+			OnValidResponsePacketDelegates.Add(ResponseDelegate);
+			OnSearchTimeoutDelegates.Add(TimeoutDelegate);
+		}
+		else
+		{
+			UE_LOG(LogTemp, Warning, TEXT("Failed to send discovery broadcast %s"), ISocketSubsystem::Get(PLATFORM_SOCKETSUBSYSTEM)->GetSocketError());
+			bSuccess = false;
+		}
+	}
+
+	return bSuccess;
+}
+
+/** Stops the LAN beacon from accepting broadcasts */
+void FLANSession::StopLANSession()
+{
+	// Don't poll anymore since we are shutting it down
+	LanBeaconState = ELanBeaconState::NotUsingLanBeacon;
+
+	// Unbind the LAN beacon object
+	if (LanBeacon)
+	{	
+		delete LanBeacon;
+		LanBeacon = NULL;
+	}
+
+	// Clear delegates
+	OnValidQueryPacketDelegates.Clear();
+	OnValidResponsePacketDelegates.Clear();
+	OnSearchTimeoutDelegates.Clear();
+
+	RetryData.Empty();
+
+	CachedResponseGuids.Empty();
+}
+
+void FLANSession::Tick(float DeltaTime)
+{
+	if (LanBeaconState == ELanBeaconState::NotUsingLanBeacon)
+	{
+		return;
+	}
+
+	uint8 PacketData[LAN_BEACON_MAX_PACKET_SIZE];
+	bool bShouldRead = true;
+	// Read each pending packet and pass it out for processing
+	while (bShouldRead)
+	{
+		int32 NumRead = LanBeacon->ReceivePacket(PacketData, LAN_BEACON_MAX_PACKET_SIZE);
+		if (NumRead > 0)
+		{
+			// Check our mode to determine the type of allowed packets
+			if (LanBeaconState == ELanBeaconState::Hosting)
+			{
+				uint64 ClientNonce;
+				// We can only accept Server Query packets
+				if (IsValidLanQueryPacket(PacketData, NumRead, ClientNonce))
+				{
+					// Strip off the header
+					OnValidQueryPacketDelegates.Broadcast(&PacketData[LAN_BEACON_PACKET_HEADER_SIZE], NumRead - LAN_BEACON_PACKET_HEADER_SIZE, ClientNonce);
+				}
+			}
+			else if (LanBeaconState == ELanBeaconState::Searching)
+			{
+				FGuid ResponseGuid;
+
+				// We can only accept Server Response packets
+				if (IsValidLanResponsePacket(PacketData, NumRead, ResponseGuid))
+				{
+					bool bAlreadySet = false;
+					CachedResponseGuids.Add(ResponseGuid, &bAlreadySet);
+
+					if (!bAlreadySet)
+					{
+						// Strip off the header
+						OnValidResponsePacketDelegates.Broadcast(&PacketData[LAN_BEACON_PACKET_HEADER_SIZE], NumRead - LAN_BEACON_PACKET_HEADER_SIZE);
+					}
+				}
+			}
+		}
+		else
+		{
+			if (LanBeaconState == ELanBeaconState::Searching)
+			{
+				LanQueryRetryTimeLeft -= DeltaTime;
+				
+				if (LanBeacon && (LanQueryRetryTimeLeft <= 0.f) && (LanQueryTimeLeft > LanQueryRetryTime))
+				{
+					if (LanBeacon->BroadcastPacket(RetryData.GetData(), RetryData.Num()))
+					{
+						UE_LOG(LogTemp, Verbose, TEXT("Sent retry query packet..."));
+					}
+					else
+					{
+						UE_LOG(LogTemp, Warning, TEXT("Failed to send retry discovery broadcast %s"), ISocketSubsystem::Get(PLATFORM_SOCKETSUBSYSTEM)->GetSocketError());
+					}
+
+					LanQueryRetryTimeLeft = LanQueryRetryTime;
+				}
+
+				// Decrement the amount of time remaining
+				LanQueryTimeLeft -= DeltaTime;
+				// Check for a timeout on the search packet
+				if (LanQueryTimeLeft <= 0.f)
+				{
+					OnSearchTimeoutDelegates.Broadcast();
+				}
+			}
+			bShouldRead = false;
+		}
+	}
+}
+
+void FLANSession::CreateHostResponsePacket(FNboSerializeToBuffer& Packet, uint64 ClientNonce)
+{
+	// Add the supported version
+	Packet << LAN_BEACON_PACKET_VERSION
+		// Platform information
+		<< (uint8)FPlatformProperties::IsLittleEndian()
+		// Game id to prevent cross game lan packets
+		<< LanGameUniqueId
+		// Add the packet type
+		<< LAN_SERVER_RESPONSE1 << LAN_SERVER_RESPONSE2
+		// Append the client nonce as a uint64
+		<< ClientNonce
+		// Append the beacon guid
+		<< LanBeaconGuid;
+}
+
+void FLANSession::CreateClientQueryPacket(FNboSerializeToBuffer& Packet, uint64 ClientNonce)
+{
+	// Build the discovery packet
+	Packet << LAN_BEACON_PACKET_VERSION
+		// Platform information
+		<< (uint8)FPlatformProperties::IsLittleEndian()
+		// Game id to prevent cross game lan packets
+		<< LanGameUniqueId
+		// Identify the packet type
+		<< LAN_SERVER_QUERY1 << LAN_SERVER_QUERY2
+		// Append the nonce as a uint64
+		<< ClientNonce
+		// Append the beacon guid
+		<< LanBeaconGuid;
+}
+
+/**
+ * Uses the cached broadcast address to send packet to a subnet
+ *
+ * @param Packet the packet to send
+ * @param Length the size of the packet to send
+ */
+bool FLANSession::BroadcastPacket(uint8* Packet, int32 Length)
+{
+	bool bSuccess = false;
+	if (LanBeacon)
+	{
+		bSuccess = LanBeacon->BroadcastPacket(Packet, Length);
+		if (!bSuccess)
+		{
+			UE_LOG(LogTemp, Warning, TEXT("Failed to send broadcast packet %d"), (int32)ISocketSubsystem::Get(PLATFORM_SOCKETSUBSYSTEM)->GetLastErrorCode());
+		}
+	}
+
+	return bSuccess;
+}
+
+/**
+ * Determines if the packet header is valid or not
+ *
+ * @param Packet the packet data to check
+ * @param Length the size of the packet buffer
+ * @param ClientNonce the client nonce contained within the packet
+ *
+ * @return true if the header is valid, false otherwise
+ */
+bool FLANSession::IsValidLanQueryPacket(const uint8* Packet, uint32 Length, uint64& ClientNonce)
+{
+	ClientNonce = 0;
+	bool bIsValid = false;
+	// Serialize out the data if the packet is the right size
+	if (Length == LAN_BEACON_PACKET_HEADER_SIZE)
+	{
+		FNboSerializeFromBuffer PacketReader(Packet,Length);
+		uint8 Version = 0;
+		PacketReader >> Version;
+		// Do the versions match?
+		if (Version == LAN_BEACON_PACKET_VERSION)
+		{
+			uint8 Platform = 255;
+			PacketReader >> Platform;
+			// Can we communicate with this platform?
+			if (Platform & LanPacketPlatformMask)
+			{
+				int32 GameId = -1;
+				PacketReader >> GameId;
+				// Is this our game?
+				if (GameId == LanGameUniqueId)
+				{
+					uint8 SQ1 = 0;
+					PacketReader >> SQ1;
+					uint8 SQ2 = 0;
+					PacketReader >> SQ2;
+					// Is this a server query?
+					bIsValid = (SQ1 == LAN_SERVER_QUERY1 && SQ2 == LAN_SERVER_QUERY2);
+					// Read the client nonce as the outvalue
+					PacketReader >> ClientNonce;
+
+					// Read the client's beacon guid 
+					FGuid BeaconGuid;
+					PacketReader >> BeaconGuid;
+				}
+			}
+		}
+	}
+	return bIsValid;
+}
+
+/**
+ * Determines if the packet header is valid or not
+ *
+ * @param Packet the packet data to check
+ * @param Length the size of the packet buffer
+ * @param ResponseGuid the lan beacon guid in the response packet
+ *
+ * @return true if the header is valid, false otherwise
+ */
+bool FLANSession::IsValidLanResponsePacket(const uint8* Packet, uint32 Length, FGuid& ResponseGuid)
+{
+	bool bIsValid = false;
+	// Serialize out the data if the packet is the right size
+	if (Length > LAN_BEACON_PACKET_HEADER_SIZE)
+	{
+		FNboSerializeFromBuffer PacketReader(Packet,Length);
+		uint8 Version = 0;
+		PacketReader >> Version;
+		// Do the versions match?
+		if (Version == LAN_BEACON_PACKET_VERSION)
+		{
+			uint8 Platform = 255;
+			PacketReader >> Platform;
+			// Can we communicate with this platform?
+			if (Platform & LanPacketPlatformMask)
+			{
+				int32 GameId = -1;
+				PacketReader >> GameId;
+				// Is this our game?
+				if (GameId == LanGameUniqueId)
+				{
+					uint8 SQ1 = 0;
+					PacketReader >> SQ1;
+					uint8 SQ2 = 0;
+					PacketReader >> SQ2;
+					// Is this a server response?
+					if (SQ1 == LAN_SERVER_RESPONSE1 && SQ2 == LAN_SERVER_RESPONSE2)
+					{
+						uint64 Nonce = 0;
+						PacketReader >> Nonce;
+						
+						if (Nonce == LanNonce)
+						{
+							FGuid BeaconGuid;
+							PacketReader >> BeaconGuid;
+
+							bIsValid = !PacketReader.HasOverflow();
+
+							if (bIsValid)
+							{
+								ResponseGuid = BeaconGuid;
+							}
+						}
+					}
+				}
+			}
+		}
+	}
+	return bIsValid;
+}
diff --git a/Engine/Plugins/Online/OnlineBase/Source/Private/Online/OnlineBaseModule.cpp b/Engine/Plugins/Online/OnlineBase/Source/Private/Online/OnlineBaseModule.cpp
new file mode 100644
--- /dev/null
+++ b/Engine/Plugins/Online/OnlineBase/Source/Private/Online/OnlineBaseModule.cpp
@@ -0,0 +1,23 @@
+// Copyright Epic Games, Inc. All Rights Reserved.
+
+#include "Modules/ModuleManager.h"
+
+/**
+ * Online base module class
+ */
+class FOnlineBaseModule : public IModuleInterface
+{
+public:
+	// IModuleInterface
+	virtual bool SupportsDynamicReloading() override
+	{
+		return false;
+	}
+
+	virtual bool SupportsAutomaticShutdown() override
+	{
+		return false;
+	}
+};
+
+IMPLEMENT_MODULE(FOnlineBaseModule, OnlineBase);
diff --git a/Engine/Plugins/Online/OnlineBase/Source/Private/OnlineBaseModule.cpp b/Engine/Plugins/Online/OnlineBase/Source/Private/OnlineBaseModule.cpp
deleted file mode 100644
diff --git a/Engine/Plugins/Online/OnlineBase/Source/Public/Online/LANBeacon.h b/Engine/Plugins/Online/OnlineBase/Source/Public/Online/LANBeacon.h
new file mode 100644
--- /dev/null
+++ b/Engine/Plugins/Online/OnlineBase/Source/Public/Online/LANBeacon.h
@@ -0,0 +1,325 @@
+// Copyright Epic Games, Inc. All Rights Reserved.
+
+#pragma once
+
+#include "CoreMinimal.h"
+#include "Misc/ConfigCacheIni.h"
+#include "Misc/Guid.h"
+
+#ifndef ONLINE_SUCCESS
+#define ONLINE_SUCCESS 0
+#endif
+
+#ifndef ONLINE_FAIL
+#define ONLINE_FAIL (uint32)-1
+#endif
+
+#ifndef ONLINE_IO_PENDING
+#define ONLINE_IO_PENDING 997
+#endif
+
+/**
+ * Generates a random nonce (number used once) of the desired length
+ *
+ * @param Nonce the buffer that will get the randomized data
+ * @param Length the number of bytes to generate random values for
+ */
+inline void GenerateNonce(uint8* Nonce, uint32 Length)
+{
+	//@todo joeg -- switch to CryptGenRandom() if possible or something equivalent
+		// Loop through generating a random value for each byte
+	for (uint32 NonceIndex = 0; NonceIndex < Length; NonceIndex++)
+	{
+		Nonce[NonceIndex] = (uint8)(FMath::Rand() & 255);
+	}
+}
+
+/**
+ * This value indicates which packet version the server is sending. Clients with
+ * differing versions will ignore these packets. This prevents crashing when
+ * changing the packet format and there are existing servers on the network
+ * Current format:
+ *
+ *	<Ver byte><Platform byte><Game unique 4 bytes><packet type 2 bytes><nonce 8 bytes><guid 16 bytes><payload>
+ */
+enum class ELANBeaconVersionHistory : uint8
+{
+	HISTORY_INITIAL				= 11,
+	HISTORY_BEACON_GUID			= 12,	// Lan response packets now contain an instance guid
+	
+	// -----<new versions can be added before this line>-------------------------------------------------
+	HISTORY_PLUS_ONE,
+	HISTORY_LATEST 							= HISTORY_PLUS_ONE - 1
+};
+
+static const uint8 LAN_BEACON_PACKET_VERSION = (uint8)ELANBeaconVersionHistory::HISTORY_LATEST;
+
+/** The size of the header for validation */
+#define LAN_BEACON_PACKET_HEADER_SIZE 32
+	
+// Offsets for various fields
+#define LAN_BEACON_VER_OFFSET 0
+#define LAN_BEACON_PLATFORM_OFFSET 1
+#define LAN_BEACON_GAMEID_OFFSET 2
+#define LAN_BEACON_PACKETTYPE1_OFFSET 6
+#define LAN_BEACON_PACKETTYPE2_OFFSET 7
+#define LAN_BEACON_NONCE_OFFSET 8
+#define LAN_BEACON_GUID_OFFSET 16
+
+// Packet types in 2 byte readable form
+#define LAN_SERVER_QUERY1 (uint8)'S'
+#define LAN_SERVER_QUERY2 (uint8)'Q'
+
+#define LAN_SERVER_RESPONSE1 (uint8)'S'
+#define LAN_SERVER_RESPONSE2 (uint8)'R'
+
+class FInternetAddr;
+class FNboSerializeToBuffer;
+
+// LAN Session Delegates
+DECLARE_MULTICAST_DELEGATE_ThreeParams(FOnValidQueryPacket, uint8*, int32, uint64);
+typedef FOnValidQueryPacket::FDelegate FOnValidQueryPacketDelegate;
+
+DECLARE_MULTICAST_DELEGATE_TwoParams(FOnValidResponsePacket, uint8*, int32);
+typedef FOnValidResponsePacket::FDelegate FOnValidResponsePacketDelegate;
+
+DECLARE_MULTICAST_DELEGATE(FOnSearchingTimeout);
+typedef FOnSearchingTimeout::FDelegate FOnSearchingTimeoutDelegate;
+
+/** Enum indicating the state the LAN beacon is in */
+namespace ELanBeaconState
+{
+	enum Type
+	{
+		/** The lan beacon is disabled */
+		NotUsingLanBeacon,
+		/** The lan beacon is responding to client requests for information */
+		Hosting,
+		/** The lan beacon is querying servers for information */
+		Searching
+	};
+
+	/** @return the stringified version of the enum passed in */
+	inline const TCHAR* ToString(ELanBeaconState::Type EnumVal)
+	{
+		switch (EnumVal)
+		{
+
+		case NotUsingLanBeacon:
+		{
+			return TEXT("NotUsingLanBeacon");
+		}
+		case Hosting:
+		{
+			return TEXT("Hosting");
+		}
+		case Searching:
+		{
+			return TEXT("Searching");
+		}
+		}
+		return TEXT("");
+	}
+}
+
+/**
+ * Class responsible for sending/receiving UDP broadcasts for LAN match
+ * discovery
+ */
+class ONLINEBASE_API FLanBeacon
+{
+	/** Builds the broadcast address and caches it */
+	TSharedPtr<class FInternetAddr> BroadcastAddr;
+	/** The socket to listen for requests on */
+	class FSocket* ListenSocket;
+	/** The address in bound requests come in on */
+	TSharedPtr<class FInternetAddr> ListenAddr;
+	/** Temporary address when receiving packets*/
+	TSharedRef<class FInternetAddr> SockAddr;
+
+public:
+	/** Sets the broadcast address for this object */
+	FLanBeacon();
+
+	/** Frees the broadcast socket */
+	~FLanBeacon();
+
+	/** Return true if there is a valid ListenSocket */
+	bool IsListenSocketValid() const;
+
+	/**
+	 * Initializes the socket
+	 *
+	 * @param Port the port to listen on
+	 *
+	 * @return true if both socket was created successfully, false otherwise
+	 */
+	bool Init(int32 Port);
+
+	/**
+	 * Called to poll the socket for pending data. Any data received is placed
+	 * in the specified packet buffer
+	 *
+	 * @param PacketData the buffer to get the socket's packet data
+	 * @param BufferSize the size of the packet buffer
+	 *
+	 * @return the number of bytes read (0 if none or an error)
+	 */
+	int32 ReceivePacket(uint8* PacketData, int32 BufferSize);
+
+	/**
+	 * Uses the cached broadcast address to send packet to a subnet
+	 *
+	 * @param Packet the packet to send
+	 * @param Length the size of the packet to send
+	 */
+	bool BroadcastPacket(uint8* Packet, int32 Length);
+};
+
+#define LAN_ANNOUNCE_PORT 14001
+#define LAN_UNIQUE_ID 9999
+#define LAN_QUERY_TIMEOUT 5
+#define LAN_QUERY_RETRY_TIME 1
+#define LAN_PLATFORMMASK 0xffffffff
+
+/**
+ *	Encapsulate functionality related to LAN broadcast data
+ */
+class ONLINEBASE_API FLANSession
+{
+protected:
+	/**
+	 * Determines if the packet header is valid or not
+	 *
+	 * @param Packet the packet data to check
+	 * @param Length the size of the packet buffer
+	 *
+	 * @return true if the header is valid, false otherwise
+	 */
+	bool IsValidLanQueryPacket(const uint8* Packet, uint32 Length, uint64& ClientNonce);
+
+	/**
+	 * Determines if the packet header is valid or not
+	 *
+	 * @param Packet the packet data to check
+	 * @param Length the size of the packet buffer
+	 * @param ResponseGuid the beacon guid from the response packet
+	 *
+	 * @return true if the header is valid, false otherwise
+	 */
+	bool IsValidLanResponsePacket(const uint8* Packet, uint32 Length, FGuid& ResponseGuid);
+
+public:
+
+	/** Port to listen on for LAN queries/responses */
+	int32 LanAnnouncePort;
+
+	/** Unique id to keep UE games from seeing each others' LAN packets */
+	int32 LanGameUniqueId;
+
+	/** Mask containing which platforms can cross communicate */
+	int32 LanPacketPlatformMask;
+
+	/** The amount of time to wait before timing out a LAN query request */
+	float LanQueryTimeout;
+
+	/** The amount of time to wait before resending a LAN query request */
+	float LanQueryRetryTime;
+
+	/** LAN beacon for packet broadcast */
+	class FLanBeacon* LanBeacon;
+
+	/** State of the LAN beacon */
+	ELanBeaconState::Type LanBeaconState;
+
+	/** Used by a client to uniquely identify itself during LAN match discovery */
+	uint64 LanNonce;
+
+	/** The amount of time before the LAN query is considered done */
+	float LanQueryTimeLeft;
+
+	/** The amount of time before the LAN query retry */
+	float LanQueryRetryTimeLeft;
+
+	/** Cached search packet data for retries */
+	TArray<uint8> RetryData;
+
+	/** Unique identifier for this hosting/searching LAN session */
+	FGuid LanBeaconGuid;
+
+	/** Cached guids already received during current search */
+	TSet<FGuid> CachedResponseGuids;
+
+	FLANSession() :
+		LanAnnouncePort(LAN_ANNOUNCE_PORT),
+		LanGameUniqueId(LAN_UNIQUE_ID),
+		LanPacketPlatformMask(LAN_PLATFORMMASK),
+		LanQueryTimeout(LAN_QUERY_TIMEOUT),
+		LanQueryRetryTime(LAN_QUERY_RETRY_TIME),
+		LanBeacon(nullptr),
+		LanBeaconState(ELanBeaconState::NotUsingLanBeacon),
+		LanNonce(0),
+		LanQueryTimeLeft(0.0f),
+		LanQueryRetryTimeLeft(0.0f)
+	{
+		if (!GConfig->GetInt(TEXT("LANSession"), TEXT("LanAnnouncePort"), LanAnnouncePort, GEngineIni))
+		{
+			LanAnnouncePort = LAN_ANNOUNCE_PORT;
+		}
+		if (!GConfig->GetInt(TEXT("LANSession"), TEXT("LanGameUniqueId"), LanGameUniqueId, GEngineIni))
+		{
+			LanGameUniqueId = LAN_UNIQUE_ID;
+		}
+	}
+
+	virtual ~FLANSession()
+	{
+		StopLANSession();
+	}
+
+	/**
+	 * Creates the LAN beacon for queries/advertising servers
+	 *
+	 * @param QueryDelegate delegate to fire when a client query is received
+	 */
+	bool Host(FOnValidQueryPacketDelegate& QueryDelegate);
+
+	/**
+	 * Creates the LAN beacon for queries/advertising servers
+	 *
+	 * @param Packet packet to be search when broadcasting the search
+	 * @param ResponseDelegate delegate to fire when a server response is received
+	 * @param TimeoutDelegate delegate to fire if we exceed maximum search time
+	 *
+	 * @return true if search was started successfully, false otherwise
+	 */
+	bool Search(class FNboSerializeToBuffer& Packet, FOnValidResponsePacketDelegate& ResponseDelegate, FOnSearchingTimeoutDelegate& TimeoutDelegate);
+
+	/**
+	 * Stops the LAN beacon from accepting broadcasts
+	 */
+	void StopLANSession();
+
+	void Tick(float DeltaTime);
+
+	/** create packet of MAX size */
+	void CreateHostResponsePacket(FNboSerializeToBuffer& Packet, uint64 ClientNonce);
+	void CreateClientQueryPacket(FNboSerializeToBuffer& Packet, uint64 ClientNonce);
+
+	/**
+	 * Uses the cached broadcast address to send packet to a subnet
+	 *
+	 * @param Packet the packet to send
+	 * @param Length the size of the packet to send
+	 */
+	bool BroadcastPacket(uint8* Packet, int32 Length);
+
+	ELanBeaconState::Type GetBeaconState() const
+	{
+		return LanBeaconState;
+	}
+
+	FOnValidQueryPacket OnValidQueryPacketDelegates;
+	FOnValidResponsePacket OnValidResponsePacketDelegates;
+	FOnSearchingTimeout OnSearchTimeoutDelegates;
+};
\ No newline at end of file
diff --git a/Engine/Plugins/Online/OnlineBase/Source/Public/Online/NboSerializer.h b/Engine/Plugins/Online/OnlineBase/Source/Public/Online/NboSerializer.h
new file mode 100644
--- /dev/null
+++ b/Engine/Plugins/Online/OnlineBase/Source/Public/Online/NboSerializer.h
@@ -0,0 +1,698 @@
+// Copyright Epic Games, Inc. All Rights Reserved.
+
+#pragma once
+
+#include "CoreMinimal.h"
+#include "Misc/Guid.h"
+#include "IPAddress.h"
+
+/**
+ * Serializes data in network byte order form into a buffer
+ */
+class ONLINEBASE_API FNboSerializeToBuffer
+{
+	/** Hidden on purpose */
+	FNboSerializeToBuffer(void);
+
+protected:
+	/**
+	 * Holds the data as it is serialized
+	 */
+	TArray<uint8> Data;
+	/**
+	 * Tracks how many bytes have been written in the packet
+	 */
+	uint32 NumBytes;
+
+	/** Indicates whether writing to the buffer caused an overflow or not */
+	bool bHasOverflowed;
+
+public:
+	/**
+	 * Initializes the write tracking
+	 */
+	FNboSerializeToBuffer(uint32 Size) :
+		NumBytes(0),
+		bHasOverflowed(false)
+	{
+		Data.Empty(Size);
+		Data.AddZeroed(Size);
+	}
+
+	/**
+	 * Cast operator to get at the formatted buffer data
+	 */
+	inline operator uint8*(void) const
+	{
+		return (uint8*)Data.GetData();
+	}
+
+	/**
+	 * Cast operator to get at the formatted buffer data
+	 */
+	inline const TArray<uint8>& GetBuffer(void) const
+	{
+		return Data;
+	}
+
+	/**
+	 * Returns the size of the buffer we've built up
+	 */
+	inline uint32 GetByteCount(void) const
+	{
+		return NumBytes;
+	}
+
+	/**
+	 * Returns the number of bytes preallocated in the array
+	 */
+	inline uint32 GetBufferSize(void) const
+	{
+		return (uint32)Data.Num();
+	}
+
+	/**
+	 * Trim any extra space in the buffer that is not used
+	 */
+	inline void TrimBuffer()
+	{
+		if (GetBufferSize() > GetByteCount())
+		{
+			Data.RemoveAt(GetByteCount(), GetBufferSize()-GetByteCount());
+		}
+	}
+
+	/**
+	 * Act as if the buffer is now empty. Useful for buffers that are reused
+	 */
+	inline void Reset()
+	{
+		NumBytes = 0;
+		bHasOverflowed = false;
+	}
+
+	/**
+	 * Adds a char to the buffer
+	 */
+	friend inline FNboSerializeToBuffer& operator<<(FNboSerializeToBuffer& Ar,const char Ch)
+	{
+		if (!Ar.HasOverflow() && Ar.NumBytes + 1 <= Ar.GetBufferSize())
+		{
+			Ar.Data[Ar.NumBytes++] = Ch;
+		}
+		else
+		{
+			Ar.bHasOverflowed = true;
+		}
+
+		return Ar;
+	}
+
+	/**
+	 * Adds a byte to the buffer
+	 */
+	friend inline FNboSerializeToBuffer& operator<<(FNboSerializeToBuffer& Ar,const uint8& B)
+	{
+		if (!Ar.HasOverflow() && Ar.NumBytes + 1 <= Ar.GetBufferSize())
+		{
+			Ar.Data[Ar.NumBytes++] = B;
+		}
+		else
+		{
+			Ar.bHasOverflowed = true;
+		}
+
+		return Ar;
+	}
+
+	/**
+	 * Adds a byte to the buffer
+	 */
+	template<class TEnum>
+	friend inline FNboSerializeToBuffer& operator<<(FNboSerializeToBuffer& Ar,const TEnumAsByte<TEnum>& B)
+	{
+		return Ar << *(uint8*)&B;
+	}
+
+	/**
+	 * Adds an int32 to the buffer
+	 */
+	friend inline FNboSerializeToBuffer& operator<<(FNboSerializeToBuffer& Ar,const int32& I)
+	{
+		return Ar << *(uint32*)&I;
+	}
+
+	/**
+	 * Adds a uint32 to the buffer
+	 */
+	friend inline FNboSerializeToBuffer& operator<<(FNboSerializeToBuffer& Ar,const uint32& D)
+	{
+		if (!Ar.HasOverflow() && Ar.NumBytes + 4 <= Ar.GetBufferSize())
+		{
+			Ar.Data[Ar.NumBytes + 0] = (D >> 24) & 0xFF;
+			Ar.Data[Ar.NumBytes + 1] = (D >> 16) & 0xFF;
+			Ar.Data[Ar.NumBytes + 2] = (D >> 8) & 0xFF;
+			Ar.Data[Ar.NumBytes + 3] = D & 0xFF;
+			Ar.NumBytes += 4;
+		}
+		else
+		{
+			Ar.bHasOverflowed = true;
+		}
+
+		return Ar;
+	}
+
+	/**
+	 * Adds a uint64 to the buffer
+	 */
+	friend inline FNboSerializeToBuffer& operator<<(FNboSerializeToBuffer& Ar,const uint64& Q)
+	{
+		if (!Ar.HasOverflow() && Ar.NumBytes + 8 <= Ar.GetBufferSize())
+		{
+			Ar.Data[Ar.NumBytes + 0] = (Q >> 56) & 0xFF;
+			Ar.Data[Ar.NumBytes + 1] = (Q >> 48) & 0xFF;
+			Ar.Data[Ar.NumBytes + 2] = (Q >> 40) & 0xFF;
+			Ar.Data[Ar.NumBytes + 3] = (Q >> 32) & 0xFF;
+			Ar.Data[Ar.NumBytes + 4] = (Q >> 24) & 0xFF;
+			Ar.Data[Ar.NumBytes + 5] = (Q >> 16) & 0xFF;
+			Ar.Data[Ar.NumBytes + 6] = (Q >> 8) & 0xFF;
+			Ar.Data[Ar.NumBytes + 7] = Q & 0xFF;
+			Ar.NumBytes += 8;
+		}
+		else
+		{
+			Ar.bHasOverflowed = true;
+		}
+
+		return Ar;
+	}
+
+	/**
+	 * Adds a float to the buffer
+	 */
+	friend inline FNboSerializeToBuffer& operator<<(FNboSerializeToBuffer& Ar,const float& F)
+	{
+		uint32 Value = *(uint32*)&F;
+		return Ar << Value;
+	}
+
+	/**
+	 * Adds a double to the buffer
+	 */
+	friend inline FNboSerializeToBuffer& operator<<(FNboSerializeToBuffer& Ar,const double& Dbl)
+	{
+		uint64 Value = *(uint64*)&Dbl;
+		return Ar << Value;
+	}
+
+	/**
+	 * Adds a FString to the buffer
+	 */
+	friend inline FNboSerializeToBuffer& operator<<(FNboSerializeToBuffer& Ar,const FString& String)
+	{
+		// We send strings length prefixed
+		FTCHARToUTF8 Converted(*String);
+		int32 Len = Converted.Length();
+
+		Ar << Len;
+
+		if (!Ar.HasOverflow() && Ar.NumBytes + Len <= Ar.GetBufferSize())
+		{
+			// Handle empty strings
+			if (Len > 0)
+			{
+				ANSICHAR* Ptr = (ANSICHAR*)Converted.Get();
+
+				// memcpy it into the buffer
+				FMemory::Memcpy(&Ar.Data[Ar.NumBytes], Ptr,Len);
+				Ar.NumBytes += Len;
+			}
+		}
+		else
+		{
+			Ar.bHasOverflowed = true;
+		}
+
+		return Ar;
+	}
+
+	/**
+	 * Adds a string to the buffer
+	 */
+	friend inline FNboSerializeToBuffer& operator<<(FNboSerializeToBuffer& Ar,const TCHAR* String)
+	{
+		// We send strings length prefixed (conversion handles null pointers)
+		FTCHARToUTF8 Converted(String);
+		int32 Len = Converted.Length();
+
+		Ar << Len;
+
+		if (!Ar.HasOverflow() && Ar.NumBytes + Len <= Ar.GetBufferSize())
+		{
+			// Handle empty/null strings
+			if (Len > 0)
+			{
+				ANSICHAR* Ptr = (ANSICHAR*)Converted.Get();
+
+				// memcpy it into the buffer
+				FMemory::Memcpy(&Ar.Data[Ar.NumBytes], Ptr, Len);
+				Ar.NumBytes += Len;
+			}
+		}
+		else
+		{
+			Ar.bHasOverflowed = true;
+		}
+
+		return Ar;
+	}
+
+	/**
+	 * Adds a substring to the buffer
+	 */
+	inline FNboSerializeToBuffer& AddString(const ANSICHAR* String,const int32 Length)
+	{
+		// We send strings length prefixed
+		int32 Len = Length;
+		(*this) << Len;
+
+		if (!HasOverflow() && NumBytes + Len <= GetBufferSize())
+		{
+			// Don't process if null
+			if (String)
+			{
+				// memcpy it into the buffer
+				FMemory::Memcpy(&Data[NumBytes],String,Len);
+				NumBytes += Len;
+			}
+		}
+		else
+		{
+			bHasOverflowed = true;
+		}
+
+		return *this;
+	}
+
+	/**
+	 * Adds an FName to the buffer as a string
+	 */
+	friend inline FNboSerializeToBuffer& operator<<(FNboSerializeToBuffer& Ar,const FName& Name)
+	{
+		const FString NameString = Name.ToString();
+		Ar << NameString;
+		return Ar;
+	}
+
+	/**
+	 * Adds an ip address to the buffer
+	 */
+	friend inline FNboSerializeToBuffer& operator<<(FNboSerializeToBuffer& Ar,const FInternetAddr& Addr)
+	{
+		TArray<uint8> RawAddressBytes = Addr.GetRawIp();
+		Ar << RawAddressBytes.Num();
+		Ar.WriteBinary(RawAddressBytes.GetData(), RawAddressBytes.Num());
+		Ar << Addr.GetPort();
+		return Ar;
+	}
+
+	/**
+	 * Adds a FGuid to the buffer
+	 */
+	friend inline FNboSerializeToBuffer& operator<<(FNboSerializeToBuffer& Ar,const FGuid& Guid)
+	{
+		Ar << Guid.A;
+		Ar << Guid.B;
+		Ar << Guid.C;
+		Ar << Guid.D;
+		return Ar;
+	}
+
+	/**
+	 * Writes a blob of data to the buffer
+	 *
+	 * @param Buffer the source data to append
+	 * @param NumToWrite the size of the blob to write
+	 */
+	inline void WriteBinary(const uint8* Buffer,uint32 NumToWrite)
+	{
+		if (!HasOverflow() && NumBytes + NumToWrite <= GetBufferSize())
+		{
+			FMemory::Memcpy(&Data[NumBytes],Buffer,NumToWrite);
+			NumBytes += NumToWrite;
+		}
+		else
+		{
+			bHasOverflowed = true;
+		}
+	}
+
+	/**
+	 * Gets the buffer at a specific point
+	 */
+	inline uint8* GetRawBuffer(uint32 Offset)
+	{
+		return &Data[Offset];
+	}
+
+	/**
+	 * Skips forward in the buffer by the specified amount
+	 *
+	 * @param Amount the number of bytes to skip ahead
+	 */
+	inline void SkipAheadBy(uint32 Amount)
+	{
+		if (!HasOverflow() && NumBytes + Amount <= GetBufferSize())
+		{
+			NumBytes += Amount;
+		}
+		else
+		{
+			bHasOverflowed = true;
+		}
+	}
+
+	/** Returns whether the buffer had an overflow when writing to it */
+	inline bool HasOverflow(void) const
+	{
+		return bHasOverflowed;
+	}
+};
+
+MSVC_PRAGMA(warning(push))
+// Disable used without initialization warning because the reads are initializing
+MSVC_PRAGMA(warning(disable : 4700))
+
+/**
+ * Class used to read data from a NBO data buffer
+ */
+class ONLINEBASE_API FNboSerializeFromBuffer
+{
+protected:
+	/** Pointer to the data this reader is attached to */
+	const uint8* Data;
+	/** The size of the data in bytes */
+	const int32 NumBytes;
+	/** The current location in the byte stream for reading */
+	int32 CurrentOffset;
+	/** Indicates whether reading from the buffer caused an overflow or not */
+	bool bHasOverflowed;
+
+	/** Hidden on purpose */
+	FNboSerializeFromBuffer(void);
+
+public:
+	/**
+	 * Initializes the buffer, size, and zeros the read offset
+	 *
+	 * @param InData the buffer to attach to
+	 * @param Length the size of the buffer we are attaching to
+	 */
+	FNboSerializeFromBuffer(const uint8* InData,int32 Length) :
+		Data(InData),
+		NumBytes(Length),
+		CurrentOffset(0),
+		bHasOverflowed(false)
+	{
+	}
+
+	/**
+	 * Reads a char from the buffer
+	 */
+	friend inline FNboSerializeFromBuffer& operator>>(FNboSerializeFromBuffer& Ar,char& Ch)
+	{
+		if (!Ar.HasOverflow() && Ar.CurrentOffset + 1 <= Ar.NumBytes)
+		{
+			Ch = Ar.Data[Ar.CurrentOffset++];
+		}
+		else
+		{
+			Ar.bHasOverflowed = true;
+		}
+		return Ar;
+	}
+
+	/**
+	 * Reads a byte from the buffer
+	 */
+	friend inline FNboSerializeFromBuffer& operator>>(FNboSerializeFromBuffer& Ar,uint8& B)
+	{
+		if (!Ar.HasOverflow() && Ar.CurrentOffset + 1 <= Ar.NumBytes)
+		{
+			B = Ar.Data[Ar.CurrentOffset++];
+		}
+		else
+		{
+			Ar.bHasOverflowed = true;
+		}
+		return Ar;
+	}
+	/**
+	 * Reads a byte from the buffer
+	 */
+	template<class TEnum>
+	friend inline FNboSerializeFromBuffer& operator>>(FNboSerializeFromBuffer& Ar,TEnumAsByte<TEnum>& B)
+	{
+		return Ar >> *(uint8*)&B;
+	}
+
+	/**
+	 * Reads an int32 from the buffer
+	 */
+	friend inline FNboSerializeFromBuffer& operator>>(FNboSerializeFromBuffer& Ar,int32& I)
+	{
+		return Ar >> *(uint32*)&I;
+	}
+
+	/**
+	 * Reads a uint32 from the buffer
+	 */
+	friend inline FNboSerializeFromBuffer& operator>>(FNboSerializeFromBuffer& Ar,uint32& D)
+	{
+		if (!Ar.HasOverflow() && Ar.CurrentOffset + 4 <= Ar.NumBytes)
+		{
+			uint32 D1 = Ar.Data[Ar.CurrentOffset + 0];
+			uint32 D2 = Ar.Data[Ar.CurrentOffset + 1];
+			uint32 D3 = Ar.Data[Ar.CurrentOffset + 2];
+			uint32 D4 = Ar.Data[Ar.CurrentOffset + 3];
+			D = D1 << 24 | D2 << 16 | D3 << 8 | D4;
+			Ar.CurrentOffset += 4;
+		}
+		else
+		{
+			Ar.bHasOverflowed = true;
+		}
+		return Ar;
+	}
+
+	/**
+	 * Adds a uint64 to the buffer
+	 */
+	friend inline FNboSerializeFromBuffer& operator>>(FNboSerializeFromBuffer& Ar,uint64& Q)
+	{
+		if (!Ar.HasOverflow() && Ar.CurrentOffset + 8 <= Ar.NumBytes)
+		{
+			Q = ((uint64)Ar.Data[Ar.CurrentOffset + 0] << 56) |
+				((uint64)Ar.Data[Ar.CurrentOffset + 1] << 48) |
+				((uint64)Ar.Data[Ar.CurrentOffset + 2] << 40) |
+				((uint64)Ar.Data[Ar.CurrentOffset + 3] << 32) |
+				((uint64)Ar.Data[Ar.CurrentOffset + 4] << 24) |
+				((uint64)Ar.Data[Ar.CurrentOffset + 5] << 16) |
+				((uint64)Ar.Data[Ar.CurrentOffset + 6] << 8) |
+				(uint64)Ar.Data[Ar.CurrentOffset + 7];
+			Ar.CurrentOffset += 8;
+		}
+		else
+		{
+			Ar.bHasOverflowed = true;
+		}
+		return Ar;
+	}
+
+	/**
+	 * Reads a float from the buffer
+	 */
+	friend inline FNboSerializeFromBuffer& operator>>(FNboSerializeFromBuffer& Ar,float& F)
+	{
+		return Ar >> *(uint32*)&F;
+	}
+
+	/**
+	 * Reads a double from the buffer
+	 */
+	friend inline FNboSerializeFromBuffer& operator>>(FNboSerializeFromBuffer& Ar,double& Dbl)
+	{
+		return Ar >> *(uint64*)&Dbl;
+	}
+
+	/**
+	 * Reads a FString from the buffer
+	 */
+	friend inline FNboSerializeFromBuffer& operator>>(FNboSerializeFromBuffer& Ar,FString& String)
+	{
+		// We send strings length prefixed
+		int32 Len = 0;
+		Ar >> Len;
+
+		// Check this way to trust NumBytes and CurrentOffset to be more accurate than the packet Len value
+		const bool bSizeOk = (Len >= 0) && (Len <= (Ar.NumBytes - Ar.CurrentOffset));
+		if (!Ar.HasOverflow() && bSizeOk)
+		{
+			// Handle strings of zero length
+			if (Len > 0)
+			{
+				char* Temp = (char*)FMemory_Alloca(Len + 1);
+				// memcpy it in from the buffer
+				FMemory::Memcpy(Temp, &Ar.Data[Ar.CurrentOffset], Len);
+				Temp[Len] = '\0';
+
+				FUTF8ToTCHAR Converted(Temp);
+				TCHAR* Ptr = (TCHAR*)Converted.Get();
+				String = Ptr;
+				Ar.CurrentOffset += Len;
+			}
+			else
+			{
+				String.Empty();
+			}
+		}
+		else
+		{
+			Ar.bHasOverflowed = true;
+		}
+
+		return Ar;
+	}
+
+	/**
+	 * Reads an FName from the buffer
+	 */
+	friend inline FNboSerializeFromBuffer& operator>>(FNboSerializeFromBuffer& Ar,FName& Name)
+	{
+		FString NameString; 
+		Ar >> NameString;
+		Name = FName(*NameString);
+		return Ar;
+	}
+
+	/**
+	 * Reads an ip address from the buffer
+	 */
+	friend inline FNboSerializeFromBuffer& operator>>(FNboSerializeFromBuffer& Ar,FInternetAddr& Addr)
+	{
+		TArray<uint8> RawAddressBytes;
+		uint32 AddressSize = 0;
+		Ar >> AddressSize;
+		Ar.ReadBinaryArray(RawAddressBytes, AddressSize);
+		Addr.SetRawIp(RawAddressBytes);
+		int32 InPort = 0;
+		Ar >> InPort;
+		Addr.SetPort(InPort);
+		return Ar;
+	}
+
+	/**
+	 * Reads a FGuid from the buffer
+	 */
+	friend inline FNboSerializeFromBuffer& operator>>(FNboSerializeFromBuffer& Ar,FGuid& Guid)
+	{
+		Ar >> Guid.A;
+		Ar >> Guid.B;
+		Ar >> Guid.C;
+		Ar >> Guid.D;
+		return Ar;
+	}
+
+	/**
+	 * Reads a blob of data from the buffer into an array
+	 * Prevents the necessity for additional allocation.
+	 *
+	 * @param OutArray the destination array
+	 * @param NumToRead the size of the blob to read
+	 */
+	void ReadBinaryArray(TArray<uint8>& OutArray, uint32 NumToRead)
+	{
+		OutArray.Reserve(NumToRead);
+		if (!HasOverflow() && CurrentOffset + (int32)NumToRead <= NumBytes)
+		{
+			for (uint32 i = 0; i < NumToRead; ++i)
+			{
+				OutArray.Add(Data[CurrentOffset + i]);
+			}
+			CurrentOffset += NumToRead;
+		}
+		else
+		{
+			bHasOverflowed = true;
+		}
+	}
+
+	/**
+	 * Reads a blob of data from the buffer
+	 *
+	 * @param OutBuffer the destination buffer
+	 * @param NumToRead the size of the blob to read
+	 */
+	void ReadBinary(uint8* OutBuffer,uint32 NumToRead)
+	{
+		if (!HasOverflow() && CurrentOffset + (int32)NumToRead <= NumBytes)
+		{
+			FMemory::Memcpy(OutBuffer,&Data[CurrentOffset],NumToRead);
+			CurrentOffset += NumToRead;
+		}
+		else
+		{
+			bHasOverflowed = true;
+		}
+	}
+
+	/**
+	 * Seek to the desired position in the buffer
+	 *
+	 * @param Pos the offset from the start of the buffer
+	 */
+	void Seek(int32 Pos)
+	{
+		checkSlow(Pos >= 0);
+
+		if (!HasOverflow() && Pos < NumBytes)
+		{
+			CurrentOffset = Pos;
+		}
+		else
+		{
+			bHasOverflowed = true;
+		}
+	}
+
+	/** @return Current position of the buffer being to be read */
+	inline int32 Tell(void) const
+	{
+		return CurrentOffset;
+	}
+
+	/** Returns whether the buffer had an overflow when reading from it */
+	inline bool HasOverflow(void) const
+	{
+		return bHasOverflowed;
+	}
+
+	/** @return Number of bytes remaining to read from the current offset to the end of the buffer */
+	inline int32 AvailableToRead(void) const
+	{
+		return FMath::Max<int32>(0,NumBytes - CurrentOffset);
+	}
+
+	/**
+	 * Returns the number of total bytes in the buffer
+	 */
+	inline int32 GetBufferSize(void) const
+	{
+		return NumBytes;
+	}
+};
+
+MSVC_PRAGMA(warning(pop))
+
+
diff --git a/Engine/Plugins/Online/OnlineSubsystem/Source/OnlineSubsystem.Build.cs b/Engine/Plugins/Online/OnlineSubsystem/Source/OnlineSubsystem.Build.cs
--- a/Engine/Plugins/Online/OnlineSubsystem/Source/OnlineSubsystem.Build.cs
+++ b/Engine/Plugins/Online/OnlineSubsystem/Source/OnlineSubsystem.Build.cs
@@ -27,7 +27,6 @@ public class OnlineSubsystem : ModuleRules
 			new string[] { 
 				"Core",
 				"CoreUObject",
-				"Sockets",
 				"JsonUtilities",
 			}
 		);
diff --git a/Engine/Plugins/Online/OnlineSubsystem/Source/Private/LANBeacon.cpp b/Engine/Plugins/Online/OnlineSubsystem/Source/Private/LANBeacon.cpp
deleted file mode 100644
diff --git a/Engine/Plugins/Online/OnlineSubsystem/Source/Private/OnlineMessage.cpp b/Engine/Plugins/Online/OnlineSubsystem/Source/Private/OnlineMessage.cpp
--- a/Engine/Plugins/Online/OnlineSubsystem/Source/Private/OnlineMessage.cpp
+++ b/Engine/Plugins/Online/OnlineSubsystem/Source/Private/OnlineMessage.cpp
@@ -9,11 +9,11 @@
 #include "Serialization/JsonSerializer.h"
 #include "OnlineKeyValuePair.h"
 #include "Interfaces/OnlineMessageInterface.h"
-#include "NboSerializer.h"
+#include "NboSerializerOSS.h"
 
 void FOnlineMessagePayload::ToBytes(TArray<uint8>& OutBytes) const
 {
-	FNboSerializeToBuffer Ar(MaxPayloadSize);
+	FNboSerializeToBufferOSS Ar(MaxPayloadSize);
 	Ar << KeyValData;
 	Ar.TrimBuffer();
 	OutBytes = Ar.GetBuffer();
@@ -21,7 +21,7 @@ void FOnlineMessagePayload::ToBytes(TArray<uint8>& OutBytes) const
 
 void FOnlineMessagePayload::FromBytes(const TArray<uint8>& InBytes)
 {
-	FNboSerializeFromBuffer Ar(InBytes.GetData(), InBytes.Num());
+	FNboSerializeFromBufferOSS Ar(InBytes.GetData(), InBytes.Num());
 	Ar >> KeyValData;
 }
 
diff --git a/Engine/Plugins/Online/OnlineSubsystem/Source/Private/OnlineSubsystem.cpp b/Engine/Plugins/Online/OnlineSubsystem/Source/Private/OnlineSubsystem.cpp
--- a/Engine/Plugins/Online/OnlineSubsystem/Source/Private/OnlineSubsystem.cpp
+++ b/Engine/Plugins/Online/OnlineSubsystem/Source/Private/OnlineSubsystem.cpp
@@ -4,10 +4,11 @@
 #include "Misc/CommandLine.h"
 #include "Misc/ConfigCacheIni.h"
 #include "HAL/IConsoleManager.h"
-#include "NboSerializer.h"
+#include "Online/NboSerializer.h"
 #include "Misc/NetworkVersion.h"
 #include "Logging/LogMacros.h"
 #include "Misc/EngineVersion.h"
+#include "OnlineSessionSettings.h"
 
 #include "Interfaces/OnlineChatInterface.h"
 #include "Interfaces/OnlinePartyInterface.h"
diff --git a/Engine/Plugins/Online/OnlineSubsystem/Source/Public/LANBeacon.h b/Engine/Plugins/Online/OnlineSubsystem/Source/Public/LANBeacon.h
deleted file mode 100644
diff --git a/Engine/Plugins/Online/OnlineSubsystem/Source/Public/NboSerializer.h b/Engine/Plugins/Online/OnlineSubsystem/Source/Public/NboSerializer.h
deleted file mode 100644
diff --git a/Engine/Plugins/Online/OnlineSubsystem/Source/Public/NboSerializerOSS.h b/Engine/Plugins/Online/OnlineSubsystem/Source/Public/NboSerializerOSS.h
new file mode 100644
--- /dev/null
+++ b/Engine/Plugins/Online/OnlineSubsystem/Source/Public/NboSerializerOSS.h
@@ -0,0 +1,270 @@
+// Copyright Epic Games, Inc. All Rights Reserved.
+
+#pragma once
+
+#include "Online/NboSerializer.h"
+
+#include "OnlineSessionSettings.h"
+#include "OnlineKeyValuePair.h"
+
+/**
+ * Serializes data in network byte order form into a buffer
+ */
+class ONLINESUBSYSTEM_API FNboSerializeToBufferOSS : public FNboSerializeToBuffer
+{
+public:
+	/** Default constructor zeros num bytes*/
+	FNboSerializeToBufferOSS() :
+		FNboSerializeToBuffer(512)
+	{
+	}
+
+	/** Constructor specifying the size to use */
+	FNboSerializeToBufferOSS(uint32 Size) :
+		FNboSerializeToBuffer(Size)
+	{
+	}
+
+	/**
+	 * Writes a list of key value pairs to buffer
+	 */
+	template<class KeyType, class ValueType>
+	friend inline FNboSerializeToBufferOSS& operator<<(FNboSerializeToBufferOSS& Ar, const FOnlineKeyValuePairs<KeyType,ValueType>& KeyValuePairs)
+	{
+		((FNboSerializeToBuffer&)Ar) << KeyValuePairs.Num();
+		for (typename FOnlineKeyValuePairs<KeyType,ValueType>::TConstIterator It(KeyValuePairs); It; ++It)
+		{
+			((FNboSerializeToBuffer&)Ar) << It.Key();
+			Ar << It.Value();
+		}
+		return Ar;
+	}
+
+	/**
+	 * Adds a key value pair to the buffer
+	 */
+	friend inline FNboSerializeToBufferOSS& operator<<(FNboSerializeToBufferOSS& Ar,const FVariantData& KeyValuePair)
+	{
+		// Write the type
+		uint8 Type = KeyValuePair.GetType();
+		((FNboSerializeToBuffer&)Ar) << Type;
+		// Now write out the held data
+		switch (KeyValuePair.GetType())
+		{
+		case EOnlineKeyValuePairDataType::Float:
+			{
+				float Value;
+				KeyValuePair.GetValue(Value);
+				((FNboSerializeToBuffer&)Ar) << Value;
+				break;
+			}
+		case EOnlineKeyValuePairDataType::Int32:
+			{
+				int32 Value;
+				KeyValuePair.GetValue(Value);
+				((FNboSerializeToBuffer&)Ar) << Value;
+				break;
+			}
+		case EOnlineKeyValuePairDataType::Int64:
+			{
+				uint64 Value;
+				KeyValuePair.GetValue(Value);
+				((FNboSerializeToBuffer&)Ar) << Value;
+				break;
+			}
+		case EOnlineKeyValuePairDataType::Double:
+			{
+				double Value;
+				KeyValuePair.GetValue(Value);
+				((FNboSerializeToBuffer&)Ar) << Value;
+				break;
+			}
+		case EOnlineKeyValuePairDataType::Blob:
+			{
+				TArray<uint8> Value;
+				KeyValuePair.GetValue(Value);
+
+				// Write the length
+				((FNboSerializeToBuffer&)Ar) << Value.Num();
+				// Followed by each byte of data
+				for (int32 Index = 0; Index < Value.Num(); Index++)
+				{
+					((FNboSerializeToBuffer&)Ar) << Value[Index];
+				}
+				break;
+			}
+		case EOnlineKeyValuePairDataType::String:
+			{
+				FString Value;
+				KeyValuePair.GetValue(Value);
+				// This will write a length prefixed string
+				((FNboSerializeToBuffer&)Ar) << (TCHAR*)*Value;
+				break;
+			}
+		case EOnlineKeyValuePairDataType::Bool:
+			{
+				bool Value;
+				KeyValuePair.GetValue(Value);
+				((FNboSerializeToBuffer&)Ar) << (uint8)(Value ? 1 : 0);
+				break;
+			}
+		case EOnlineKeyValuePairDataType::Empty:
+			break;
+		default:
+			checkfSlow(false, TEXT("Unsupported EOnlineKeyValuePairDataType: %d"), Type);
+		}
+		return Ar;
+	}
+
+	/**
+	 * Adds a single online session setting to the buffer
+	 */
+	friend inline FNboSerializeToBufferOSS& operator<<(FNboSerializeToBufferOSS& Ar,const FOnlineSessionSetting& Setting)
+	{
+		Ar << Setting.Data;
+		
+		uint8 Type;
+		Type = Setting.AdvertisementType;
+		((FNboSerializeToBuffer&)Ar) << Type;
+
+		return Ar;
+	}
+};
+
+/**
+ * Class used to read data from a NBO data buffer
+ */
+class ONLINESUBSYSTEM_API FNboSerializeFromBufferOSS : public FNboSerializeFromBuffer
+{
+public:
+	/**
+	 * Initializes the buffer, size, and zeros the read offset
+	 */
+	FNboSerializeFromBufferOSS(const uint8* Packet,int32 Length) :
+		FNboSerializeFromBuffer(Packet,Length)
+	{
+	}
+
+	/**
+	 * Reads a list of key value pairs from the buffer
+	 */
+	template<class KeyType, class ValueType>
+	friend inline FNboSerializeFromBufferOSS& operator>>(FNboSerializeFromBufferOSS& Ar, FOnlineKeyValuePairs<KeyType,ValueType>& KeyValuePairs)
+	{
+		int32 NumValues;
+		Ar >> NumValues;		
+		for (int32 Idx=0; Idx < NumValues; Idx++)
+		{
+			KeyType Key;
+			ValueType Value;
+			Ar >> Key;
+			Ar >> Value;
+
+			KeyValuePairs.Add(Key,Value);
+		}
+		return Ar;
+	}
+	
+	/**
+	 * Reads a key value pair from the buffer
+	 */
+	friend inline FNboSerializeFromBufferOSS& operator>>(FNboSerializeFromBufferOSS& Ar,FVariantData& KeyValuePair)
+	{
+		if (!Ar.HasOverflow())
+		{
+			// Read the type
+			uint8 Type;
+			Ar >> Type;
+
+			// Now read in the held data
+			switch ((EOnlineKeyValuePairDataType::Type)Type)
+			{
+			case EOnlineKeyValuePairDataType::Float:
+				{
+					float Value;
+					Ar >> Value;
+					KeyValuePair.SetValue(Value);
+					break;
+				}
+			case EOnlineKeyValuePairDataType::Int32:
+				{
+					int32 Value;
+					Ar >> Value;
+					KeyValuePair.SetValue(Value);
+					break;
+				}
+			case EOnlineKeyValuePairDataType::Int64:
+				{
+					uint64 Value;
+					Ar >> Value;
+					KeyValuePair.SetValue(Value);
+					break;
+				}
+			case EOnlineKeyValuePairDataType::Double:
+				{
+					double Value;
+					Ar >> Value;
+					KeyValuePair.SetValue(Value);
+					break;
+				}
+			case EOnlineKeyValuePairDataType::Blob:
+				{
+					int32 Length;
+					Ar >> Length;
+
+					// Check this way to trust NumBytes and CurrentOffset to be more accurate than the packet Len value
+					const bool bSizeOk = (Length >= 0) && (Length <= (Ar.NumBytes - Ar.CurrentOffset));
+					if (!Ar.HasOverflow() && bSizeOk)
+					{
+						// Now directly copy the blob data
+						KeyValuePair.SetValue(Length, &Ar.Data[Ar.CurrentOffset]);
+						Ar.CurrentOffset += Length;
+					}
+					else
+					{
+						Ar.bHasOverflowed = true;
+					}
+
+					break;
+				}
+			case EOnlineKeyValuePairDataType::String:
+				{
+					FString Value;
+					Ar >> Value;
+					KeyValuePair.SetValue(Value);
+					break;
+				}
+			case EOnlineKeyValuePairDataType::Bool:
+				{
+					uint8 Value;
+					Ar >> Value;
+					KeyValuePair.SetValue(Value != 0);
+					break;
+				}
+			case EOnlineKeyValuePairDataType::Empty:
+				break;
+			default:
+				checkfSlow(false, TEXT("Unsupported EOnlineKeyValuePairDataType: %d"), Type);
+			}
+		}
+
+		return Ar;
+	}
+
+	/**
+	 * Reads a single session setting from the buffer
+	 */
+	friend inline FNboSerializeFromBufferOSS& operator>>(FNboSerializeFromBufferOSS& Ar,FOnlineSessionSetting& Setting)
+	{
+		Ar >> Setting.Data;
+
+		if (!Ar.HasOverflow())
+		{
+			uint8 Type;
+			Ar >> Type;
+			Setting.AdvertisementType = (EOnlineDataAdvertisementType::Type) Type;
+		}
+
+		return Ar;
+	}
+};
\ No newline at end of file
diff --git a/Engine/Plugins/Online/OnlineSubsystem/Source/Public/OnlineSubsystemTypes.h b/Engine/Plugins/Online/OnlineSubsystem/Source/Public/OnlineSubsystemTypes.h
--- a/Engine/Plugins/Online/OnlineSubsystem/Source/Public/OnlineSubsystemTypes.h
+++ b/Engine/Plugins/Online/OnlineSubsystem/Source/Public/OnlineSubsystemTypes.h
@@ -37,22 +37,6 @@ extern ONLINESUBSYSTEM_API bool IsUniqueIdLocal(const FUniqueNetId& UniqueId);
 #define ONLINE_IO_PENDING 997
 #endif
 
-/**
- * Generates a random nonce (number used once) of the desired length
- *
- * @param Nonce the buffer that will get the randomized data
- * @param Length the number of bytes to generate random values for
- */
-inline void GenerateNonce(uint8* Nonce, uint32 Length)
-{
-//@todo joeg -- switch to CryptGenRandom() if possible or something equivalent
-	// Loop through generating a random value for each byte
-	for (uint32 NonceIndex = 0; NonceIndex < Length; NonceIndex++)
-	{
-		Nonce[NonceIndex] = (uint8)(FMath::Rand() & 255);
-	}
-}
-
 /**
  * Environment for the current online platform
  */
@@ -433,42 +417,6 @@ namespace ELeaderboardUpdateMethod
 	}
 }
 
-/** Enum indicating the state the LAN beacon is in */
-namespace ELanBeaconState
-{
-	enum Type
-	{
-		/** The lan beacon is disabled */
-		NotUsingLanBeacon,
-		/** The lan beacon is responding to client requests for information */
-		Hosting,
-		/** The lan beacon is querying servers for information */
-		Searching
-	};
-
-	/** @return the stringified version of the enum passed in */
-	inline const TCHAR* ToString(ELanBeaconState::Type EnumVal)
-	{
-		switch (EnumVal)
-		{
-
-		case NotUsingLanBeacon:
-			{
-				return TEXT("NotUsingLanBeacon");
-			}
-		case Hosting:
-			{
-				return TEXT("Hosting");
-			}
-		case Searching:
-			{
-				return TEXT("Searching");
-			}
-		}
-		return TEXT("");
-	}
-}
-
 /** Enum indicating the current state of the online session (in progress, ended, etc.) */
 namespace EOnlineSessionState
 {
diff --git a/Engine/Plugins/Online/OnlineSubsystemEOS/Source/OnlineSubsystemEOS/OnlineSubsystemEOS.Build.cs b/Engine/Plugins/Online/OnlineSubsystemEOS/Source/OnlineSubsystemEOS/OnlineSubsystemEOS.Build.cs
--- a/Engine/Plugins/Online/OnlineSubsystemEOS/Source/OnlineSubsystemEOS/OnlineSubsystemEOS.Build.cs
+++ b/Engine/Plugins/Online/OnlineSubsystemEOS/Source/OnlineSubsystemEOS/OnlineSubsystemEOS.Build.cs
@@ -25,6 +25,7 @@ public class OnlineSubsystemEOS : ModuleRules
 				"EOSShared",
 				"EOSVoiceChat",
 				"Json",
+				"OnlineBase",
 				"OnlineSubsystem",
 				"Sockets",
 				"SocketSubsystemEOS",
diff --git a/Engine/Plugins/Online/OnlineSubsystemEOS/Source/OnlineSubsystemEOS/Private/NboSerializerEOS.h b/Engine/Plugins/Online/OnlineSubsystemEOS/Source/OnlineSubsystemEOS/Private/NboSerializerEOS.h
--- a/Engine/Plugins/Online/OnlineSubsystemEOS/Source/OnlineSubsystemEOS/Private/NboSerializerEOS.h
+++ b/Engine/Plugins/Online/OnlineSubsystemEOS/Source/OnlineSubsystemEOS/Private/NboSerializerEOS.h
@@ -4,25 +4,25 @@
 
 #include "CoreMinimal.h"
 #include "OnlineSubsystemEOSTypes.h"
-#include "NboSerializer.h"
+#include "NboSerializerOSS.h"
 
 #if WITH_EOS_SDK
 
 /**
  * Serializes data in network byte order form into a buffer
  */
-class FNboSerializeToBufferEOS : public FNboSerializeToBuffer
+class FNboSerializeToBufferEOS : public FNboSerializeToBufferOSS
 {
 public:
 	/** Default constructor zeros num bytes*/
 	FNboSerializeToBufferEOS() :
-		FNboSerializeToBuffer(512)
+		FNboSerializeToBufferOSS(512)
 	{
 	}
 
 	/** Constructor specifying the size to use */
 	FNboSerializeToBufferEOS(uint32 Size) :
-		FNboSerializeToBuffer(Size)
+		FNboSerializeToBufferOSS(Size)
 	{
 	}
 
@@ -34,7 +34,7 @@ public:
 		check(SessionInfo.HostAddr.IsValid());
 		// Skip SessionType (assigned at creation)
 		Ar << *SessionInfo.SessionId;
-		Ar << *SessionInfo.HostAddr;
+		((FNboSerializeToBuffer&)Ar) << *SessionInfo.HostAddr;
 		return Ar;
  	}
 
@@ -43,7 +43,7 @@ public:
 	 */
 	friend inline FNboSerializeToBufferEOS& operator<<(FNboSerializeToBufferEOS& Ar, const FUniqueNetIdEOS& UniqueId)
 	{
-		Ar << UniqueId.UniqueNetIdStr;
+		((FNboSerializeToBuffer&)Ar) << UniqueId.UniqueNetIdStr;
 		return Ar;
 	}
 
@@ -52,7 +52,7 @@ public:
 	 */
 	friend inline FNboSerializeToBufferEOS& operator<<(FNboSerializeToBufferEOS& Ar, const FUniqueNetIdString& UniqueId)
 	{
-		Ar << UniqueId.UniqueNetIdStr;
+		((FNboSerializeToBuffer&)Ar) << UniqueId.UniqueNetIdStr;
 		return Ar;
 	}
 };
@@ -60,14 +60,14 @@ public:
 /**
  * Class used to write data into packets for sending via system link
  */
-class FNboSerializeFromBufferEOS : public FNboSerializeFromBuffer
+class FNboSerializeFromBufferEOS : public FNboSerializeFromBufferOSS
 {
 public:
 	/**
 	 * Initializes the buffer, size, and zeros the read offset
 	 */
 	FNboSerializeFromBufferEOS(uint8* Packet,int32 Length) :
-		FNboSerializeFromBuffer(Packet,Length)
+		FNboSerializeFromBufferOSS(Packet,Length)
 	{
 	}
 
diff --git a/Engine/Plugins/Online/OnlineSubsystemEOS/Source/OnlineSubsystemEOS/Private/OnlineSessionEOS.cpp b/Engine/Plugins/Online/OnlineSubsystemEOS/Source/OnlineSubsystemEOS/Private/OnlineSessionEOS.cpp
--- a/Engine/Plugins/Online/OnlineSubsystemEOS/Source/OnlineSubsystemEOS/Private/OnlineSessionEOS.cpp
+++ b/Engine/Plugins/Online/OnlineSubsystemEOS/Source/OnlineSubsystemEOS/Private/OnlineSessionEOS.cpp
@@ -2996,7 +2996,7 @@ void FOnlineSessionEOS::TickLanTasks(float DeltaTime)
 void FOnlineSessionEOS::AppendSessionToPacket(FNboSerializeToBufferEOS& Packet, FOnlineSession* Session)
 {
 	/** Owner of the session */
-	Packet << *StaticCastSharedPtr<const FUniqueNetIdEOS>(Session->OwningUserId)
+	((FNboSerializeToBuffer&)Packet) << Session->OwningUserId->ToString()
 		<< Session->OwningUserName
 		<< Session->NumOpenPrivateConnections
 		<< Session->NumOpenPublicConnections;
@@ -3018,7 +3018,7 @@ void FOnlineSessionEOS::AppendSessionSettingsToPacket(FNboSerializeToBufferEOS&
 #endif 
 
 	// Members of the session settings class
-	Packet << SessionSettings->NumPublicConnections
+	((FNboSerializeToBuffer&)Packet) << SessionSettings->NumPublicConnections
 		<< SessionSettings->NumPrivateConnections
 		<< (uint8)SessionSettings->bShouldAdvertise
 		<< (uint8)SessionSettings->bIsLANMatch
@@ -3044,13 +3044,13 @@ void FOnlineSessionEOS::AppendSessionSettingsToPacket(FNboSerializeToBufferEOS&
 	}
 
 	// Add count of advertised keys and the data
-	Packet << (int32)NumAdvertisedProperties;
+	((FNboSerializeToBuffer&)Packet) << (int32)NumAdvertisedProperties;
 	for (FSessionSettings::TConstIterator It(SessionSettings->Settings); It; ++It)
 	{
 		const FOnlineSessionSetting& Setting = It.Value();
 		if (Setting.AdvertisementType >= EOnlineDataAdvertisementType::ViaOnlineService)
 		{
-			Packet << It.Key();
+			((FNboSerializeToBuffer&)Packet) << It.Key();
 			Packet << Setting;
 #if DEBUG_LAN_BEACON
 			UE_LOG_ONLINE_SESSION(Verbose, TEXT("%s"), *Setting.ToString());
diff --git a/Engine/Plugins/Online/OnlineSubsystemEOS/Source/OnlineSubsystemEOS/Private/OnlineSessionEOS.h b/Engine/Plugins/Online/OnlineSubsystemEOS/Source/OnlineSubsystemEOS/Private/OnlineSessionEOS.h
--- a/Engine/Plugins/Online/OnlineSubsystemEOS/Source/OnlineSubsystemEOS/Private/OnlineSessionEOS.h
+++ b/Engine/Plugins/Online/OnlineSubsystemEOS/Source/OnlineSubsystemEOS/Private/OnlineSessionEOS.h
@@ -5,7 +5,7 @@
 #include "Misc/ScopeLock.h"
 #include "OnlineSessionSettings.h"
 #include "Interfaces/OnlineSessionInterface.h"
-#include "LANBeacon.h"
+#include "Online/LANBeacon.h"
 #include "OnlineSubsystemEOSTypes.h"
 
 class FOnlineSubsystemEOS;
diff --git a/Engine/Plugins/Online/OnlineSubsystemNull/Source/OnlineSubsystemNull.Build.cs b/Engine/Plugins/Online/OnlineSubsystemNull/Source/OnlineSubsystemNull.Build.cs
--- a/Engine/Plugins/Online/OnlineSubsystemNull/Source/OnlineSubsystemNull.Build.cs
+++ b/Engine/Plugins/Online/OnlineSubsystemNull/Source/OnlineSubsystemNull.Build.cs
@@ -31,6 +31,7 @@ public class OnlineSubsystemNull : ModuleRules
 				"Core",
 				"Sockets",
 				"CoreUObject",
+				"OnlineBase",
 				"OnlineSubsystem", 
 				"Json"
 			}
diff --git a/Engine/Plugins/Online/OnlineSubsystemNull/Source/Private/NboSerializerNull.h b/Engine/Plugins/Online/OnlineSubsystemNull/Source/Private/NboSerializerNull.h
--- a/Engine/Plugins/Online/OnlineSubsystemNull/Source/Private/NboSerializerNull.h
+++ b/Engine/Plugins/Online/OnlineSubsystemNull/Source/Private/NboSerializerNull.h
@@ -4,23 +4,23 @@
 
 #include "CoreMinimal.h"
 #include "OnlineSubsystemNullTypes.h"
-#include "NboSerializer.h"
+#include "NboSerializerOSS.h"
 
 /**
  * Serializes data in network byte order form into a buffer
  */
-class FNboSerializeToBufferNull : public FNboSerializeToBuffer
+class ONLINESUBSYSTEMNULL_API FNboSerializeToBufferNull : public FNboSerializeToBufferOSS
 {
 public:
 	/** Default constructor zeros num bytes*/
 	FNboSerializeToBufferNull() :
-		FNboSerializeToBuffer(512)
+		FNboSerializeToBufferOSS(512)
 	{
 	}
 
 	/** Constructor specifying the size to use */
 	FNboSerializeToBufferNull(uint32 Size) :
-		FNboSerializeToBuffer(Size)
+		FNboSerializeToBufferOSS(Size)
 	{
 	}
 
@@ -32,7 +32,7 @@ public:
 		check(SessionInfo.HostAddr.IsValid());
 		// Skip SessionType (assigned at creation)
 		Ar << *SessionInfo.SessionId;
-		Ar << *SessionInfo.HostAddr;
+		((FNboSerializeToBuffer&)Ar) << *SessionInfo.HostAddr;
 		return Ar;
  	}
 
@@ -41,7 +41,7 @@ public:
 	 */
 	friend inline FNboSerializeToBufferNull& operator<<(FNboSerializeToBufferNull& Ar, const FUniqueNetIdNull& UniqueId)
 	{
-		Ar << UniqueId.UniqueNetIdStr;
+		((FNboSerializeToBuffer&)Ar) << UniqueId.UniqueNetIdStr;
 		return Ar;
 	}
 };
@@ -49,14 +49,14 @@ public:
 /**
  * Class used to write data into packets for sending via system link
  */
-class FNboSerializeFromBufferNull : public FNboSerializeFromBuffer
+class ONLINESUBSYSTEMNULL_API FNboSerializeFromBufferNull : public FNboSerializeFromBufferOSS
 {
 public:
 	/**
 	 * Initializes the buffer, size, and zeros the read offset
 	 */
 	FNboSerializeFromBufferNull(uint8* Packet,int32 Length) :
-		FNboSerializeFromBuffer(Packet,Length)
+		FNboSerializeFromBufferOSS(Packet,Length)
 	{
 	}
 
diff --git a/Engine/Plugins/Online/OnlineSubsystemNull/Source/Private/OnlineSessionInterfaceNull.cpp b/Engine/Plugins/Online/OnlineSubsystemNull/Source/Private/OnlineSessionInterfaceNull.cpp
--- a/Engine/Plugins/Online/OnlineSubsystemNull/Source/Private/OnlineSessionInterfaceNull.cpp
+++ b/Engine/Plugins/Online/OnlineSubsystemNull/Source/Private/OnlineSessionInterfaceNull.cpp
@@ -1000,7 +1000,7 @@ void FOnlineSessionNull::TickLanTasks(float DeltaTime)
 void FOnlineSessionNull::AppendSessionToPacket(FNboSerializeToBufferNull& Packet, FOnlineSession* Session)
 {
 	/** Owner of the session */
-	Packet << *StaticCastSharedPtr<const FUniqueNetIdNull>(Session->OwningUserId)
+	((FNboSerializeToBuffer&) Packet) << Session->OwningUserId->ToString()
 		<< Session->OwningUserName
 		<< Session->NumOpenPrivateConnections
 		<< Session->NumOpenPublicConnections;
@@ -1022,7 +1022,7 @@ void FOnlineSessionNull::AppendSessionSettingsToPacket(FNboSerializeToBufferNull
 #endif 
 
 	// Members of the session settings class
-	Packet << SessionSettings->NumPublicConnections
+	((FNboSerializeToBuffer&)Packet) << SessionSettings->NumPublicConnections
 		<< SessionSettings->NumPrivateConnections
 		<< (uint8)SessionSettings->bShouldAdvertise
 		<< (uint8)SessionSettings->bIsLANMatch
@@ -1048,13 +1048,13 @@ void FOnlineSessionNull::AppendSessionSettingsToPacket(FNboSerializeToBufferNull
 	}
 
 	// Add count of advertised keys and the data
-	Packet << (int32)NumAdvertisedProperties;
+	((FNboSerializeToBuffer&)Packet) << (int32)NumAdvertisedProperties;
 	for (FSessionSettings::TConstIterator It(SessionSettings->Settings); It; ++It)
 	{
 		const FOnlineSessionSetting& Setting = It.Value();
 		if (Setting.AdvertisementType >= EOnlineDataAdvertisementType::ViaOnlineService)
 		{
-			Packet << It.Key();
+			((FNboSerializeToBuffer&)Packet) << It.Key();
 			Packet << Setting;
 #if DEBUG_LAN_BEACON
 			UE_LOG_ONLINE_SESSION(Verbose, TEXT("%s"), *Setting.ToString());
diff --git a/Engine/Plugins/Online/OnlineSubsystemNull/Source/Private/OnlineSessionInterfaceNull.h b/Engine/Plugins/Online/OnlineSubsystemNull/Source/Private/OnlineSessionInterfaceNull.h
--- a/Engine/Plugins/Online/OnlineSubsystemNull/Source/Private/OnlineSessionInterfaceNull.h
+++ b/Engine/Plugins/Online/OnlineSubsystemNull/Source/Private/OnlineSessionInterfaceNull.h
@@ -8,7 +8,7 @@
 #include "OnlineSessionSettings.h"
 #include "Interfaces/OnlineSessionInterface.h"
 #include "OnlineSubsystemNullPackage.h"
-#include "LANBeacon.h"
+#include "Online/LANBeacon.h"
 
 class FOnlineSubsystemNull;
 
diff --git a/Engine/Plugins/Online/OnlineSubsystemSteam/Source/OnlineSubsystemSteam.Build.cs b/Engine/Plugins/Online/OnlineSubsystemSteam/Source/OnlineSubsystemSteam.Build.cs
--- a/Engine/Plugins/Online/OnlineSubsystemSteam/Source/OnlineSubsystemSteam.Build.cs
+++ b/Engine/Plugins/Online/OnlineSubsystemSteam/Source/OnlineSubsystemSteam.Build.cs
@@ -25,6 +25,7 @@ public class OnlineSubsystemSteam : ModuleRules
 				"Sockets", 
 				"Voice",
                 "AudioMixer",
+				"OnlineBase",
 				"OnlineSubsystem",
 				"Json",
 				"PacketHandler",
diff --git a/Engine/Plugins/Online/OnlineSubsystemSteam/Source/Private/NboSerializerSteam.h b/Engine/Plugins/Online/OnlineSubsystemSteam/Source/Private/NboSerializerSteam.h
--- a/Engine/Plugins/Online/OnlineSubsystemSteam/Source/Private/NboSerializerSteam.h
+++ b/Engine/Plugins/Online/OnlineSubsystemSteam/Source/Private/NboSerializerSteam.h
@@ -4,23 +4,23 @@
 
 #include "CoreMinimal.h"
 #include "OnlineSubsystemSteamTypes.h"
-#include "NboSerializer.h"
+#include "NboSerializerOSS.h"
 
 /**
  * Serializes data in network byte order form into a buffer
  */
-class FNboSerializeToBufferSteam : public FNboSerializeToBuffer
+class FNboSerializeToBufferSteam : public FNboSerializeToBufferOSS
 {
 public:
 	/** Default constructor zeros num bytes*/
 	FNboSerializeToBufferSteam() :
-		FNboSerializeToBuffer(512)
+		FNboSerializeToBufferOSS(512)
 	{
 	}
 
 	/** Constructor specifying the size to use */
 	FNboSerializeToBufferSteam(uint32 Size) :
-		FNboSerializeToBuffer(Size)
+		FNboSerializeToBufferOSS(Size)
 	{
 	}
 
@@ -31,8 +31,8 @@ public:
  	{
 		check(SessionInfo.HostAddr.IsValid());
 		// Skip SessionType (assigned at creation)
-		Ar << *SessionInfo.HostAddr;
-		Ar << SessionInfo.SessionId->UniqueNetId;
+		((FNboSerializeToBuffer&)Ar) << *SessionInfo.HostAddr;
+		((FNboSerializeToBuffer&)Ar) << SessionInfo.SessionId->UniqueNetId;
 		return Ar;
  	}
 };
@@ -40,14 +40,14 @@ public:
 /**
  * Class used to write data into packets for sending via system link
  */
-class FNboSerializeFromBufferSteam : public FNboSerializeFromBuffer
+class FNboSerializeFromBufferSteam : public FNboSerializeFromBufferOSS
 {
 public:
 	/**
 	 * Initializes the buffer, size, and zeros the read offset
 	 */
 	FNboSerializeFromBufferSteam(uint8* Packet,int32 Length) :
-		FNboSerializeFromBuffer(Packet,Length)
+		FNboSerializeFromBufferOSS(Packet,Length)
 	{
 	}
 
diff --git a/Engine/Plugins/Online/OnlineSubsystemSteam/Source/Private/OnlineSessionInterfaceSteam.cpp b/Engine/Plugins/Online/OnlineSubsystemSteam/Source/Private/OnlineSessionInterfaceSteam.cpp
--- a/Engine/Plugins/Online/OnlineSubsystemSteam/Source/Private/OnlineSessionInterfaceSteam.cpp
+++ b/Engine/Plugins/Online/OnlineSubsystemSteam/Source/Private/OnlineSessionInterfaceSteam.cpp
@@ -10,7 +10,7 @@
 #include "OnlineSessionAsyncServerSteam.h"
 #include "OnlineLeaderboardInterfaceSteam.h"
 #include "OnlineAuthInterfaceSteam.h"
-#include "LANBeacon.h"
+#include "Online/LANBeacon.h"
 #include "NboSerializerSteam.h"
 #include "Interfaces/VoiceInterface.h"
 
@@ -1582,7 +1582,7 @@ void FOnlineSessionSteam::TickPendingInvites(float DeltaTime)
 void FOnlineSessionSteam::AppendSessionToPacket(FNboSerializeToBufferSteam& Packet, FOnlineSession* Session)
 {
 	/** Owner of the session */
-	Packet << StaticCastSharedPtr<const FUniqueNetIdSteam>(Session->OwningUserId)->UniqueNetId
+	((FNboSerializeToBuffer&)Packet) << StaticCastSharedPtr<const FUniqueNetIdSteam>(Session->OwningUserId)->UniqueNetId
 		<< Session->OwningUserName
 		<< Session->NumOpenPrivateConnections
 		<< Session->NumOpenPublicConnections;
@@ -1601,7 +1601,7 @@ void FOnlineSessionSteam::AppendSessionSettingsToPacket(FNboSerializeToBufferSte
 #endif 
 
 	// Members of the session settings class
-	Packet << SessionSettings->NumPublicConnections
+	((FNboSerializeToBuffer&)Packet) << SessionSettings->NumPublicConnections
 		<< SessionSettings->NumPrivateConnections
 		<< (uint8)SessionSettings->bShouldAdvertise
 		<< (uint8)SessionSettings->bIsLANMatch
@@ -1627,13 +1627,13 @@ void FOnlineSessionSteam::AppendSessionSettingsToPacket(FNboSerializeToBufferSte
 	}
 
 	// Add count of advertised keys and the data
-	Packet << (int32)NumAdvertisedProperties;
+	((FNboSerializeToBuffer&)Packet) << (int32)NumAdvertisedProperties;
 	for (FSessionSettings::TConstIterator It(SessionSettings->Settings); It; ++It)
 	{
 		const FOnlineSessionSetting& Setting = It.Value();
 		if (Setting.AdvertisementType >= EOnlineDataAdvertisementType::ViaOnlineService)
 		{
-			Packet << It.Key();
+			((FNboSerializeToBuffer&)Packet) << It.Key();
 			Packet << Setting;
 #if DEBUG_LAN_BEACON
 			UE_LOG_ONLINE_SESSION(Verbose, TEXT("%s"), *Setting.ToString());
