Added templated case to catch unsupported types in NboSerializer.
Fix for unsupported type in NboSerializerOSS.

diff --git a/Engine/Plugins/Online/OnlineBase/Source/Public/Online/NboSerializer.h b/Engine/Plugins/Online/OnlineBase/Source/Public/Online/NboSerializer.h
--- a/Engine/Plugins/Online/OnlineBase/Source/Public/Online/NboSerializer.h
+++ b/Engine/Plugins/Online/OnlineBase/Source/Public/Online/NboSerializer.h
@@ -6,6 +6,12 @@
 #include "Misc/Guid.h"
 #include "IPAddress.h"
 
+namespace NboSerializer
+{
+	template <typename T>
+	static constexpr inline bool UnexpectedType = false;
+}
+
 /**
  * Serializes data in network byte order form into a buffer
  */
@@ -339,6 +345,15 @@ public:
 		return Ar;
 	}
 
+	// This will be chosen in case of any type not specially supported
+	template <typename T>
+	friend inline FNboSerializeToBuffer& operator <<(FNboSerializeToBuffer& Ar, const T& Value)
+	{
+		static_assert(NboSerializer::UnexpectedType<T>, "Unsupported type in FNboSerializeToBuffer::operator<<");
+
+		return Ar;
+	}
+
 	/**
 	 * Writes a blob of data to the buffer
 	 *
@@ -626,6 +641,15 @@ public:
 		return Ar;
 	}
 
+	// This will be chosen in case of any type not specially supported
+	template <typename T>
+	friend inline FNboSerializeFromBuffer& operator>>(FNboSerializeFromBuffer& Ar, T& Value)
+	{
+		static_assert(NboSerializer::UnexpectedType<T>, "Unsupported type in FNboSerializeFromBuffer::operator>>");
+
+		return Ar;
+	}
+
 	/**
 	 * Reads a blob of data from the buffer into an array
 	 * Prevents the necessity for additional allocation.
diff --git a/Engine/Plugins/Online/OnlineSubsystem/Source/Public/NboSerializerOSS.h b/Engine/Plugins/Online/OnlineSubsystem/Source/Public/NboSerializerOSS.h
--- a/Engine/Plugins/Online/OnlineSubsystem/Source/Public/NboSerializerOSS.h
+++ b/Engine/Plugins/Online/OnlineSubsystem/Source/Public/NboSerializerOSS.h
@@ -98,7 +98,7 @@ public:
 				FString Value;
 				KeyValuePair.GetValue(Value);
 				// This will write a length prefixed string
-				((FNboSerializeToBuffer&)Ar) << (TCHAR*)*Value;
+				((FNboSerializeToBuffer&)Ar) << *Value;
 				break;
 			}
 		case EOnlineKeyValuePairDataType::Bool:
