Add SETTING_ACTIVITY_SESSION to control setting a session as the activity session

diff --git a/Engine/Plugins/Online/OnlineBase/Source/Public/Online/OnlineSessionNames.h b/Engine/Plugins/Online/OnlineBase/Source/Public/Online/OnlineSessionNames.h
--- a/Engine/Plugins/Online/OnlineBase/Source/Public/Online/OnlineSessionNames.h
+++ b/Engine/Plugins/Online/OnlineBase/Source/Public/Online/OnlineSessionNames.h
@@ -110,6 +110,8 @@
 #define SETTING_CUSTOM_JOIN_INFO FName(TEXT("CUSTOMJOININFO"))
 /** What verbosity of session updates to subscribe to (value is FString, comma-separated if multiple values needed)*/
 #define SETTING_SESSION_SUBSCRIPTION_TYPES FName(TEXT("SESSIONSUBSCRIPTIONTYPE"))
+/** Whether the session should be considered as the player's activity session or otherwise (value is bool) */
+#define SETTING_ACTIVITY_SESSION FName(TEXT("ACTIVITYSESSION"))
 
 /**
  * Search settings
