deprecate online session defines part 2

diff --git a/Engine/Plugins/Online/OnlineBase/Source/Public/Online/OnlineSessionNames.h b/Engine/Plugins/Online/OnlineBase/Source/Public/Online/OnlineSessionNames.h
--- a/Engine/Plugins/Online/OnlineBase/Source/Public/Online/OnlineSessionNames.h
+++ b/Engine/Plugins/Online/OnlineBase/Source/Public/Online/OnlineSessionNames.h
@@ -138,9 +138,9 @@
 /** The matchmaking queue name to matchmake in, e.g. "TeamDeathmatch" (value is string) */
 #define SEARCH_MATCHMAKING_QUEUE FName(TEXT("MATCHMAKINGQUEUE"))
 /** If set, use the named Xbox Live hopper to find a session via matchmaking (value is a string) */
-#define SEARCH_XBOX_LIVE_HOPPER_NAME FName(TEXT("LIVEHOPPERNAME"))
-/** Which session template from the service configuration to use */
-#define SEARCH_XBOX_LIVE_SESSION_TEMPLATE_NAME FName(TEXT("LIVESESSIONTEMPLATE"))
+#define SEARCH_XBOX_LIVE_HOPPER_NAME UE_DEPRECATED_MACRO(5.4, "SEARCH_XBOX_LIVE_HOPPER_NAME has been deprecated. Use SETTING_MATCHING_HOPPER instead.") SETTING_MATCHING_HOPPER
+/** Which session template from the service configuration to use.*/
+#define SEARCH_XBOX_LIVE_SESSION_TEMPLATE_NAME UE_DEPRECATED_MACRO(5.4, "SEARCH_XBOX_LIVE_SESSION_TEMPLATE_NAME has been deprecated. Use SETTING_SESSION_TEMPLATE_NAME instead.")  SETTING_SESSION_TEMPLATE_NAME
 /** Selection method used to determine which match to join when multiple are returned (valid only on Switch) */
 #define SEARCH_SWITCH_SELECTION_METHOD FName(TEXT("SWITCHSELECTIONMETHOD"))
 /** Whether to use lobbies vs sessions */
