Add a new session settings define for allowing non-primary platform members in a supported session

diff --git a/Engine/Plugins/Online/OnlineBase/Source/Public/Online/OnlineSessionNames.h b/Engine/Plugins/Online/OnlineBase/Source/Public/Online/OnlineSessionNames.h
--- a/Engine/Plugins/Online/OnlineBase/Source/Public/Online/OnlineSessionNames.h
+++ b/Engine/Plugins/Online/OnlineBase/Source/Public/Online/OnlineSessionNames.h
@@ -53,6 +53,9 @@
 #define SETTING_CUSTOMSEARCHINT7 FName(TEXT("CUSTOMSEARCHINT7"))
 #define SETTING_CUSTOMSEARCHINT8 FName(TEXT("CUSTOMSEARCHINT8"))
 
+/** Setting to allow non-primary platform members in a session */
+#define SETTING_ALLOWEXTERNALMEMBERSINSESSION FName("ALLOWEXTERNALMEMBERSINSESSION")
+
 /**
  * Setting describing whether host migration is enabled for the session.  This typically means when the host leaves, 
  * the online service picks a new host instead of destroying the session.  Not supported by all online services.
