Fix missing legal boilerplate in OnlineBase.cpp

diff --git a/Engine/Plugins/Online/OnlineBase/Source/Public/Online/OnlineBase.cpp b/Engine/Plugins/Online/OnlineBase/Source/Public/Online/OnlineBase.cpp
--- a/Engine/Plugins/Online/OnlineBase/Source/Public/Online/OnlineBase.cpp
+++ b/Engine/Plugins/Online/OnlineBase/Source/Public/Online/OnlineBase.cpp
@@ -1,4 +1,6 @@
-﻿#include "OnlineBase.h"
+﻿// Copyright Epic Games, Inc. All Rights Reserved.
+
+#include "OnlineBase.h"
 
 static TAutoConsoleVariable<int32> CVarBuildIdOverride(
 TEXT("buildidoverride2"),
