Add OnlineInit dependency resolver to the OnlineBase Plugin.

diff --git a/Engine/Plugins/Online/EOSShared/Config/Engine.ini b/Engine/Plugins/Online/EOSShared/Config/Engine.ini
new file mode 100644
--- /dev/null
+++ b/Engine/Plugins/Online/EOSShared/Config/Engine.ini
@@ -0,0 +1,2 @@
+[OnlineBase]
++InitDependencies=EOSSDK
\ No newline at end of file
diff --git a/Engine/Plugins/Online/EOSShared/EOSShared.uplugin b/Engine/Plugins/Online/EOSShared/EOSShared.uplugin
--- a/Engine/Plugins/Online/EOSShared/EOSShared.uplugin
+++ b/Engine/Plugins/Online/EOSShared/EOSShared.uplugin
@@ -31,6 +31,10 @@
 			"Name": "SteamShared",
 			"Enabled": false,
 			"Optional": true
+		},
+		{
+			"Name": "OnlineBase",
+			"Enabled": true
 		}
 	]
 }
diff --git a/Engine/Plugins/Online/EOSShared/Source/EOSShared/EOSShared.Build.cs b/Engine/Plugins/Online/EOSShared/Source/EOSShared/EOSShared.Build.cs
--- a/Engine/Plugins/Online/EOSShared/Source/EOSShared/EOSShared.Build.cs
+++ b/Engine/Plugins/Online/EOSShared/Source/EOSShared/EOSShared.Build.cs
@@ -29,7 +29,8 @@ public class EOSShared : ModuleRules
 			new string[]
 			{
 				"Core",
-				"EOSSDK"
+				"EOSSDK",
+				"OnlineBase"
 			}
 		);
 
diff --git a/Engine/Plugins/Online/EOSShared/Source/EOSShared/Private/EOSSDKManager.cpp b/Engine/Plugins/Online/EOSShared/Source/EOSShared/Private/EOSSDKManager.cpp
--- a/Engine/Plugins/Online/EOSShared/Source/EOSShared/Private/EOSSDKManager.cpp
+++ b/Engine/Plugins/Online/EOSShared/Source/EOSShared/Private/EOSSDKManager.cpp
@@ -22,6 +22,7 @@
 #include "Misc/NetworkVersion.h"
 #include "Misc/Paths.h"
 #include "Misc/ConfigCacheIni.h"
+#include "Online/OnlineBase.h"
 #include "ProfilingDebugging/CsvProfiler.h"
 #include "ProfilingDebugging/CallstackTrace.h"
 #include "Stats/Stats.h"
@@ -152,6 +153,11 @@ namespace
 #endif // !NO_LOGGING
 }
 
+namespace EOnlineInitDependencies
+{
+const FName EOSSDK = TEXT("EOSSDK");
+}
+
 #if EOSSDK_RUNTIME_LOAD_REQUIRED
 void FEOSSDKManager::LoadSDKHandle()
 {
@@ -307,6 +313,7 @@ EOS_EResult FEOSSDKManager::Initialize()
 			UE_LOG(LogEOSShared, Warning, TEXT("%hs EOS_Initialize failed error:%s"), __FUNCTION__, *LexToString(EosResult));
 		}
 
+		FOnlineInitDependencies::NotifyDependencyResolved(EOnlineInitDependencies::EOSSDK);
 		OnPostInitializeSDK.Broadcast(EosResult);
 		return EosResult;
 	}
diff --git a/Engine/Plugins/Online/OnlineBase/Source/Private/Online/InitDependencyList.cpp b/Engine/Plugins/Online/OnlineBase/Source/Private/Online/InitDependencyList.cpp
new file mode 100644
--- /dev/null
+++ b/Engine/Plugins/Online/OnlineBase/Source/Private/Online/InitDependencyList.cpp
@@ -0,0 +1,123 @@
+// Copyright Epic Games, Inc. All Rights Reserved.
+
+#include "InitDependencyList.h"
+
+#include "Algo/Transform.h"
+#include "HAL/PlatformTime.h"
+#include "Logging/LogMacros.h"
+#include "Logging/StructuredLog.h"
+#include "Misc/ConfigCacheIni.h"
+#include "Misc/CoreDelegates.h"
+
+DEFINE_LOG_CATEGORY_STATIC(LogInitDependencies, Log, All);
+
+#define INITDEPENDENCIES_LOG_NOARG(this, Level, EventName) \
+UE_LOGFMT(LogInitDependencies, Level, "ConfigSection=[{config_section}],Event=[{event_name}]:", UE_LOGFMT_VALUE(ConfigSection), UE_LOGFMT_VALUE(EventName))
+
+#define INITDEPENDENCIES_LOG(this, Level, EventName, Fmt, ...) \
+UE_LOGFMT(LogInitDependencies, Level, "ConfigSection=[{config_section}],Event=[{event_name}]:" Fmt, UE_LOGFMT_VALUE(ConfigSection), UE_LOGFMT_VALUE(EventName), ##__VA_ARGS__)
+
+namespace
+{
+
+const double LogFrequencySecondsInitial = 30;
+
+/* anonymous */ }
+
+FInitDependencyList::FInitDependencyList(const FString& InConfigSection)
+	: ConfigSection(InConfigSection)
+	, WaitingDependencies()
+	, InitStartTime(FPlatformTime::Seconds())
+	, LogFrequencySeconds(LogFrequencySecondsInitial)
+{
+	LoadConfig();
+
+	if (WaitingDependencies.IsEmpty())
+	{
+		INITDEPENDENCIES_LOG(this, Verbose, "all_dependencies_resolved", "No configured dependencies.");
+	}
+	else
+	{
+		INITDEPENDENCIES_LOG(this, Log, "dependencies_initialized", "Waiting dependencies=[{waiting_dependencies}]", UE_LOGFMT_VALUE(FString::JoinBy(WaitingDependencies, TEXT(","), UE_PROJECTION_MEMBER(FName, ToString))));
+
+		if (IsReady())
+		{
+			SignalReady();
+		}
+		else
+		{
+			ScheduleUpdateLog();
+		}
+	}
+}
+
+FInitDependencyList::~FInitDependencyList()
+{
+	CancelUpdateLog();
+}
+
+bool FInitDependencyList::IsReady()
+{
+	return WaitingDependencies.IsEmpty();
+}
+
+void FInitDependencyList::NotifyDependencyResolved(FName Dependency)
+{
+	if (WaitingDependencies.Contains(Dependency))
+	{
+		WaitingDependencies.Remove(Dependency);
+		INITDEPENDENCIES_LOG(this, Verbose, "dependency_resolved", "Resolved dependency=[{resolved_dependency}], Waiting dependencies=[{waiting_dependencies}]", UE_LOGFMT_VALUE(Dependency.ToString()), UE_LOGFMT_VALUE(FString::JoinBy(WaitingDependencies, TEXT(","), UE_PROJECTION_MEMBER(FName, ToString))));
+
+		if (IsReady())
+		{
+			SignalReady();
+		}
+	}
+}
+
+void FInitDependencyList::LoadConfig()
+{
+	WaitingDependencies.Empty();
+	TArray<FString> InInitDependencies;
+	GConfig->GetArray(*ConfigSection, TEXT("InitDependencies"), InInitDependencies, GEngineIni);
+	Algo::Transform(InInitDependencies, WaitingDependencies, [](const FString& Dependency)
+		{
+			return FName(Dependency);
+		});
+
+	double InLogFrequencySeconds = LogFrequencySecondsInitial;
+	GConfig->GetDouble(*ConfigSection, TEXT("InitDependenciesLogFrequencySeconds"), InLogFrequencySeconds, GEngineIni);
+	LogFrequencySeconds = FMath::Max(InLogFrequencySeconds, 1);
+}
+
+void FInitDependencyList::SignalReady()
+{
+	CancelUpdateLog();
+	const double Elapsed = FPlatformTime::Seconds() - InitStartTime;
+	INITDEPENDENCIES_LOG(this, Log, "all_dependencies_resolved", "Elapsed time=[{elapsed}]", UE_LOGFMT_VALUE(Elapsed));
+	FOnlineInitDependencies::OnReady.Broadcast();
+}
+
+void FInitDependencyList::ScheduleUpdateLog()
+{
+	CancelUpdateLog();
+	LogTimerDelegateHandle = FTSTicker::GetCoreTicker().AddTicker(FTickerDelegate::CreateRaw(this, &FInitDependencyList::OnUpdateLogTimerElapsed), LogFrequencySeconds);
+}
+
+void FInitDependencyList::CancelUpdateLog()
+{
+	FTSTicker::GetCoreTicker().RemoveTicker(LogTimerDelegateHandle);
+	LogTimerDelegateHandle = FTSTicker::FDelegateHandle();
+}
+
+bool FInitDependencyList::OnUpdateLogTimerElapsed(float Dt)
+{
+	if (!IsReady())
+	{
+		const double Elapsed = FPlatformTime::Seconds() - InitStartTime;
+		INITDEPENDENCIES_LOG(this, Verbose, "wait_timer_update", "Waiting dependencies=[{waiting_dependencies}], Elapsed time=[{elapsed}]", UE_LOGFMT_VALUE(FString::JoinBy(WaitingDependencies, TEXT(","), UE_PROJECTION_MEMBER(FName, ToString))), UE_LOGFMT_VALUE(Elapsed));
+		ScheduleUpdateLog();
+	}
+
+	return false;
+}
diff --git a/Engine/Plugins/Online/OnlineBase/Source/Private/Online/InitDependencyList.h b/Engine/Plugins/Online/OnlineBase/Source/Private/Online/InitDependencyList.h
new file mode 100644
--- /dev/null
+++ b/Engine/Plugins/Online/OnlineBase/Source/Private/Online/InitDependencyList.h
@@ -0,0 +1,31 @@
+// Copyright Epic Games, Inc. All Rights Reserved.
+
+#pragma once
+
+#include "Online/OnlineBase.h"
+
+#include "Containers/Ticker.h"
+
+class FInitDependencyList  final
+{
+public:
+	FInitDependencyList(const FString& InConfigSection);
+	~FInitDependencyList();
+
+	bool IsReady();
+	void NotifyDependencyResolved(FName Dependency);
+
+private:
+	void LoadConfig();
+	void SignalReady();
+
+	void ScheduleUpdateLog();
+	void CancelUpdateLog();
+	bool OnUpdateLogTimerElapsed(float Dt);
+
+	FString ConfigSection;
+	TSet<FName> WaitingDependencies;
+	double InitStartTime;
+	double LogFrequencySeconds;
+	FTSTicker::FDelegateHandle LogTimerDelegateHandle;
+};
diff --git a/Engine/Plugins/Online/OnlineBase/Source/Public/Online/OnlineBase.cpp b/Engine/Plugins/Online/OnlineBase/Source/Public/Online/OnlineBase.cpp
--- a/Engine/Plugins/Online/OnlineBase/Source/Public/Online/OnlineBase.cpp
+++ b/Engine/Plugins/Online/OnlineBase/Source/Public/Online/OnlineBase.cpp
@@ -1,13 +1,38 @@
-﻿// Copyright Epic Games, Inc. All Rights Reserved.
+// Copyright Epic Games, Inc. All Rights Reserved.
 
 #include "OnlineBase.h"
 
+#include "Online/InitDependencyList.h"
+
 static TAutoConsoleVariable<int32> CVarBuildIdOverride(
 TEXT("buildidoverride"),
 0,
 TEXT("Sets build id used for matchmaking "));
 
+namespace
+{
+
+FInitDependencyList& GetOnlineInitDependencyResolver()
+{
+	static FInitDependencyList Impl(TEXT("OnlineBase"));
+	return Impl;
+}
+
+/* Anonymous */ }
+
 TAutoConsoleVariable<int32>& GetBuildIdOverrideCVar()
 {
 	return CVarBuildIdOverride;
 }
+
+FOnlineInitDependencies::FOnReady FOnlineInitDependencies::OnReady;
+
+bool FOnlineInitDependencies::IsReady()
+{
+	return GetOnlineInitDependencyResolver().IsReady();
+}
+
+void FOnlineInitDependencies::NotifyDependencyResolved(FName Dependency)
+{
+	GetOnlineInitDependencyResolver().NotifyDependencyResolved(Dependency);
+}
diff --git a/Engine/Plugins/Online/OnlineBase/Source/Public/Online/OnlineBase.h b/Engine/Plugins/Online/OnlineBase/Source/Public/Online/OnlineBase.h
--- a/Engine/Plugins/Online/OnlineBase/Source/Public/Online/OnlineBase.h
+++ b/Engine/Plugins/Online/OnlineBase/Source/Public/Online/OnlineBase.h
@@ -2,6 +2,8 @@
 
 #pragma once
 
+#include "Delegates/Delegate.h"
+#include "Delegates/DelegateCombinations.h"
 #include "HAL/Platform.h"
 #include "HAL/IConsoleManager.h"
 
@@ -17,4 +19,38 @@
 #define ONLINE_IO_PENDING 997
 #endif
 
-ONLINEBASE_API TAutoConsoleVariable<int32>&  GetBuildIdOverrideCVar();
\ No newline at end of file
+ONLINEBASE_API TAutoConsoleVariable<int32>&  GetBuildIdOverrideCVar();
+
+/*
+ * OnlineInitDependencies blocks online framework components from initializing until the
+ * configured dependencies of online systems are ready. Waiting for the dependencies to be ready
+ * prevents the creation of OnlineService instances using the wrong config. This functionality can
+ * be used for waiting for the EOSSDK to be initialized following process forking, or for waiting
+ * for the application to modify online config based on runtime changes such as environment detection.
+ * 
+ * To add a dependency, add an entry to the InitDependencies array under the OnlineBase section in DefaultEngine.ini.
+ * Example:
+ * [OnlineBase]
+ * +InitDependencies=RuntimeConfig
+ * 
+ * Online framework components will then be blocked from initializing until NotifyDependencyResolved(TEXT("RuntimeConfig")) is called.
+ */
+class FOnlineInitDependencies
+{
+public:
+	/*
+	 * Returns true when all dependencies have been resolved.
+	 */
+	static ONLINEBASE_API bool IsReady();
+
+	/*
+	 * Mark a dependency as resolved to allow init to proceed.
+	 */
+	static ONLINEBASE_API void NotifyDependencyResolved(FName Dependency);
+
+	/*
+	 * Delegate fired when all online init dependencies have been resolved.
+	 */
+	DECLARE_MULTICAST_DELEGATE(FOnReady);
+	static ONLINEBASE_API FOnReady OnReady;
+};
