[Backout] - CL39884527
#rnx
Original CL Desc
-----------------------------------------------------------------
Add OnlineInit dependency resolver to the OnlineBase Plugin.

diff --git a/Engine/Plugins/Online/EOSShared/Config/Engine.ini b/Engine/Plugins/Online/EOSShared/Config/Engine.ini
deleted file mode 100644
diff --git a/Engine/Plugins/Online/EOSShared/EOSShared.uplugin b/Engine/Plugins/Online/EOSShared/EOSShared.uplugin
--- a/Engine/Plugins/Online/EOSShared/EOSShared.uplugin
+++ b/Engine/Plugins/Online/EOSShared/EOSShared.uplugin
@@ -31,10 +31,6 @@
 			"Name": "SteamShared",
 			"Enabled": false,
 			"Optional": true
-		},
-		{
-			"Name": "OnlineBase",
-			"Enabled": true
 		}
 	]
 }
diff --git a/Engine/Plugins/Online/EOSShared/Source/EOSShared/EOSShared.Build.cs b/Engine/Plugins/Online/EOSShared/Source/EOSShared/EOSShared.Build.cs
--- a/Engine/Plugins/Online/EOSShared/Source/EOSShared/EOSShared.Build.cs
+++ b/Engine/Plugins/Online/EOSShared/Source/EOSShared/EOSShared.Build.cs
@@ -29,8 +29,7 @@ public class EOSShared : ModuleRules
 			new string[]
 			{
 				"Core",
-				"EOSSDK",
-				"OnlineBase"
+				"EOSSDK"
 			}
 		);
 
diff --git a/Engine/Plugins/Online/EOSShared/Source/EOSShared/Private/EOSSDKManager.cpp b/Engine/Plugins/Online/EOSShared/Source/EOSShared/Private/EOSSDKManager.cpp
--- a/Engine/Plugins/Online/EOSShared/Source/EOSShared/Private/EOSSDKManager.cpp
+++ b/Engine/Plugins/Online/EOSShared/Source/EOSShared/Private/EOSSDKManager.cpp
@@ -22,7 +22,6 @@
 #include "Misc/NetworkVersion.h"
 #include "Misc/Paths.h"
 #include "Misc/ConfigCacheIni.h"
-#include "Online/OnlineBase.h"
 #include "ProfilingDebugging/CsvProfiler.h"
 #include "ProfilingDebugging/CallstackTrace.h"
 #include "Stats/Stats.h"
@@ -153,11 +152,6 @@ namespace
 #endif // !NO_LOGGING
 }
 
-namespace EOnlineInitDependencies
-{
-const FName EOSSDK = TEXT("EOSSDK");
-}
-
 #if EOSSDK_RUNTIME_LOAD_REQUIRED
 void FEOSSDKManager::LoadSDKHandle()
 {
@@ -313,7 +307,6 @@ EOS_EResult FEOSSDKManager::Initialize()
 			UE_LOG(LogEOSShared, Warning, TEXT("%hs EOS_Initialize failed error:%s"), __FUNCTION__, *LexToString(EosResult));
 		}
 
-		FOnlineInitDependencies::NotifyDependencyResolved(EOnlineInitDependencies::EOSSDK);
 		OnPostInitializeSDK.Broadcast(EosResult);
 		return EosResult;
 	}
diff --git a/Engine/Plugins/Online/OnlineBase/Source/Private/Online/InitDependencyList.cpp b/Engine/Plugins/Online/OnlineBase/Source/Private/Online/InitDependencyList.cpp
deleted file mode 100644
diff --git a/Engine/Plugins/Online/OnlineBase/Source/Private/Online/InitDependencyList.h b/Engine/Plugins/Online/OnlineBase/Source/Private/Online/InitDependencyList.h
deleted file mode 100644
diff --git a/Engine/Plugins/Online/OnlineBase/Source/Public/Online/OnlineBase.cpp b/Engine/Plugins/Online/OnlineBase/Source/Public/Online/OnlineBase.cpp
--- a/Engine/Plugins/Online/OnlineBase/Source/Public/Online/OnlineBase.cpp
+++ b/Engine/Plugins/Online/OnlineBase/Source/Public/Online/OnlineBase.cpp
@@ -1,38 +1,13 @@
-// Copyright Epic Games, Inc. All Rights Reserved.
+﻿// Copyright Epic Games, Inc. All Rights Reserved.
 
 #include "OnlineBase.h"
 
-#include "Online/InitDependencyList.h"
-
 static TAutoConsoleVariable<int32> CVarBuildIdOverride(
 TEXT("buildidoverride"),
 0,
 TEXT("Sets build id used for matchmaking "));
 
-namespace
-{
-
-FInitDependencyList& GetOnlineInitDependencyResolver()
-{
-	static FInitDependencyList Impl(TEXT("OnlineBase"));
-	return Impl;
-}
-
-/* Anonymous */ }
-
 TAutoConsoleVariable<int32>& GetBuildIdOverrideCVar()
 {
 	return CVarBuildIdOverride;
 }
-
-FOnlineInitDependencies::FOnReady FOnlineInitDependencies::OnReady;
-
-bool FOnlineInitDependencies::IsReady()
-{
-	return GetOnlineInitDependencyResolver().IsReady();
-}
-
-void FOnlineInitDependencies::NotifyDependencyResolved(FName Dependency)
-{
-	GetOnlineInitDependencyResolver().NotifyDependencyResolved(Dependency);
-}
diff --git a/Engine/Plugins/Online/OnlineBase/Source/Public/Online/OnlineBase.h b/Engine/Plugins/Online/OnlineBase/Source/Public/Online/OnlineBase.h
--- a/Engine/Plugins/Online/OnlineBase/Source/Public/Online/OnlineBase.h
+++ b/Engine/Plugins/Online/OnlineBase/Source/Public/Online/OnlineBase.h
@@ -2,8 +2,6 @@
 
 #pragma once
 
-#include "Delegates/Delegate.h"
-#include "Delegates/DelegateCombinations.h"
 #include "HAL/Platform.h"
 #include "HAL/IConsoleManager.h"
 
@@ -19,38 +17,4 @@
 #define ONLINE_IO_PENDING 997
 #endif
 
-ONLINEBASE_API TAutoConsoleVariable<int32>&  GetBuildIdOverrideCVar();
-
-/*
- * OnlineInitDependencies blocks online framework components from initializing until the
- * configured dependencies of online systems are ready. Waiting for the dependencies to be ready
- * prevents the creation of OnlineService instances using the wrong config. This functionality can
- * be used for waiting for the EOSSDK to be initialized following process forking, or for waiting
- * for the application to modify online config based on runtime changes such as environment detection.
- * 
- * To add a dependency, add an entry to the InitDependencies array under the OnlineBase section in DefaultEngine.ini.
- * Example:
- * [OnlineBase]
- * +InitDependencies=RuntimeConfig
- * 
- * Online framework components will then be blocked from initializing until NotifyDependencyResolved(TEXT("RuntimeConfig")) is called.
- */
-class FOnlineInitDependencies
-{
-public:
-	/*
-	 * Returns true when all dependencies have been resolved.
-	 */
-	static ONLINEBASE_API bool IsReady();
-
-	/*
-	 * Mark a dependency as resolved to allow init to proceed.
-	 */
-	static ONLINEBASE_API void NotifyDependencyResolved(FName Dependency);
-
-	/*
-	 * Delegate fired when all online init dependencies have been resolved.
-	 */
-	DECLARE_MULTICAST_DELEGATE(FOnReady);
-	static ONLINEBASE_API FOnReady OnReady;
-};
+ONLINEBASE_API TAutoConsoleVariable<int32>&  GetBuildIdOverrideCVar();
\ No newline at end of file
